#ttp 330510 - EDITOR: MATINEE: Curve editor scrollbar disappears when the user tries to use the scrollbar
#proj UE4
#branch UE4
#summary The scrollbar is now sized correctly based on the height of the curve editor viewport

diff --git a/Engine/Source/Editor/DistCurveEditor/Private/SCurveEditorViewport.cpp b/Engine/Source/Editor/DistCurveEditor/Private/SCurveEditorViewport.cpp
--- a/Engine/Source/Editor/DistCurveEditor/Private/SCurveEditorViewport.cpp
+++ b/Engine/Source/Editor/DistCurveEditor/Private/SCurveEditorViewport.cpp
@@ -57,6 +57,8 @@ void SCurveEditorViewport::Construct(const FArguments& InArgs)
 	// The viewport widget needs an interface so it knows what should render
 	ViewportWidget->SetViewportInterface( Viewport.ToSharedRef() );
 
+	PrevViewportHeight = 0;
+	AdjustScrollBar();
 }
 
 void SCurveEditorViewport::RefreshViewport()
@@ -72,6 +74,12 @@ void SCurveEditorViewport::SetVerticalScrollBarPosition(float Position)
 	OnViewportVerticalScrollBarScrolled(MaxOffset * Position);
 }
 
+void SCurveEditorViewport::AdjustScrollBar()
+{
+	// Pretend the scrollbar was scrolled by no amount; this will update the scrollbar ratio
+	OnViewportVerticalScrollBarScrolled(0.0f);
+}
+
 TSharedPtr<FSceneViewport> SCurveEditorViewport::GetViewport() const
 {
 	return Viewport;
@@ -92,6 +100,19 @@ TSharedPtr<SScrollBar> SCurveEditorViewport::GetVerticalScrollBar() const
 	return ViewportVerticalScrollBar;
 }
 
+void SCurveEditorViewport::Tick( const FGeometry& AllottedGeometry, const double InCurrentTime, const float InDeltaTime )
+{
+	SCompoundWidget::Tick(AllottedGeometry, InCurrentTime, InDeltaTime);
+
+	// Check to see if we need to update the scrollbars due to a size change
+	const int32 CurrentHeight = AllottedGeometry.Size.Y;
+	if (CurrentHeight != PrevViewportHeight)
+	{
+		PrevViewportHeight = CurrentHeight;
+		AdjustScrollBar();
+	}
+}
+
 EVisibility SCurveEditorViewport::GetViewportVerticalScrollBarVisibility() const
 {
 	return (ViewportClient->GetViewportVerticalScrollBarRatio() < 1.0f)? EVisibility::Visible: EVisibility::Collapsed;
diff --git a/Engine/Source/Editor/DistCurveEditor/Private/SCurveEditorViewport.h b/Engine/Source/Editor/DistCurveEditor/Private/SCurveEditorViewport.h
--- a/Engine/Source/Editor/DistCurveEditor/Private/SCurveEditorViewport.h
+++ b/Engine/Source/Editor/DistCurveEditor/Private/SCurveEditorViewport.h
@@ -36,6 +36,12 @@ public:
 	/** Scrolls the scrollbar... expected range is from 0.0 to 1.0 */
 	void SetVerticalScrollBarPosition(float Position);
 
+	/**
+	 * Updates the scroll bar for the current state of the window's size and content layout.  This should be called
+	 * when either the window size changes or the vertical size of the content contained in the window changes.
+	 */
+	void AdjustScrollBar();
+
 	/** Accessors */
 	TSharedPtr<FSceneViewport> GetViewport() const;
 	TSharedPtr<FCurveEditorViewportClient> GetViewportClient() const;
@@ -43,6 +49,8 @@ public:
 	TSharedPtr<SScrollBar> GetVerticalScrollBar() const;
 
 protected:
+	virtual void Tick( const FGeometry& AllottedGeometry, const double InCurrentTime, const float InDeltaTime ) OVERRIDE;
+
 	/** Returns the visibility of the viewport scrollbars */
 	EVisibility GetViewportVerticalScrollBarVisibility() const;
 
@@ -64,4 +72,7 @@ private:
 
 	/** Vertical scrollbar */
 	TSharedPtr<SScrollBar> ViewportVerticalScrollBar;
+
+	/** Height of the viewport on the last call to Tick */
+	int32 PrevViewportHeight;
 };
