#ttp 331438 UE4: PlatformMath refactor
#proj core
#branch UE4
#summary Rename Trunc, Round, Floor, Ceil to TruncToInt, RoundToInt, FloorToInt, CeilToInt. Added *ToFloat versions. Repeated for FGenericPlatformMath and all derived classes for all platforms wherever applicable. Corrected comment in Fractional and added Frac() which does HLSL-style fractional (x - floor(x)). Checked for compilation on all projects (with cooking levels wherever applicable). Didn't change Fractional to Frac, this will be done in second commit.

diff --git a/Engine/Plugins/Editor/EpicSurvey/Source/EpicSurvey/Private/SurveyTitleCdnStorage.cpp b/Engine/Plugins/Editor/EpicSurvey/Source/EpicSurvey/Private/SurveyTitleCdnStorage.cpp
--- a/Engine/Plugins/Editor/EpicSurvey/Source/EpicSurvey/Private/SurveyTitleCdnStorage.cpp
+++ b/Engine/Plugins/Editor/EpicSurvey/Source/EpicSurvey/Private/SurveyTitleCdnStorage.cpp
@@ -311,7 +311,7 @@ void FSurveyTitleCdnStorage::EnumerateFiles_HttpRequestComplete(FHttpRequestPtr
 						}
 						if (JsonFileHeader->HasField(TEXT("length")))
 						{
-							FileHeader.FileSize = FMath::Trunc(JsonFileHeader->GetNumberField(TEXT("length")));
+							FileHeader.FileSize = FMath::TruncToInt(JsonFileHeader->GetNumberField(TEXT("length")));
 						}
 
 						if (FileHeader.FileName.IsEmpty())
diff --git a/Engine/Plugins/Experimental/LiveEditor/Source/LiveEditor/Private/LiveEditorKismetLibrary.cpp b/Engine/Plugins/Experimental/LiveEditor/Source/LiveEditor/Private/LiveEditorKismetLibrary.cpp
--- a/Engine/Plugins/Experimental/LiveEditor/Source/LiveEditor/Private/LiveEditorKismetLibrary.cpp
+++ b/Engine/Plugins/Experimental/LiveEditor/Source/LiveEditor/Private/LiveEditorKismetLibrary.cpp
@@ -44,7 +44,7 @@ namespace nLiveEditorKismetLibrary
 		}
 		else
 		{
-			T NewValue = CurValue+(T)FMath::Round(Delta);
+			T NewValue = CurValue+(T)FMath::RoundToInt(Delta);
 			return NewValue;
 		}
 	}
diff --git a/Engine/Plugins/Experimental/Paper2D/Source/Paper2DEditor/Private/SPaperEditorViewport.cpp b/Engine/Plugins/Experimental/Paper2D/Source/Paper2DEditor/Private/SPaperEditorViewport.cpp
--- a/Engine/Plugins/Experimental/Paper2D/Source/Paper2DEditor/Private/SPaperEditorViewport.cpp
+++ b/Engine/Plugins/Experimental/Paper2D/Source/Paper2DEditor/Private/SPaperEditorViewport.cpp
@@ -247,8 +247,8 @@ FReply SPaperEditorViewport::OnMouseButtonUp(const FGeometry& MyGeometry, const
 			const FVector2D ScreenSpaceCursorPos = MyGeometry.LocalToAbsolute( GraphCoordToPanelCoord( SoftwareCursorPosition ) );
 
 			FIntPoint BestPositionInViewport(
-				FMath::Round( FMath::Clamp( ScreenSpaceCursorPos.X, ThisPanelScreenSpaceRect.Left, ThisPanelScreenSpaceRect.Right ) ),
-				FMath::Round( FMath::Clamp( ScreenSpaceCursorPos.Y, ThisPanelScreenSpaceRect.Top, ThisPanelScreenSpaceRect.Bottom ) )
+				FMath::RoundToInt( FMath::Clamp( ScreenSpaceCursorPos.X, ThisPanelScreenSpaceRect.Left, ThisPanelScreenSpaceRect.Right ) ),
+				FMath::RoundToInt( FMath::Clamp( ScreenSpaceCursorPos.Y, ThisPanelScreenSpaceRect.Top, ThisPanelScreenSpaceRect.Bottom ) )
 				);
 
 			if (!bCursorInDeadZone)
@@ -368,8 +368,8 @@ FReply SPaperEditorViewport::OnMouseMove(const FGeometry& MyGeometry, const FPoi
 // 
 // 					// Snap to grid
 // 					const float SnapSize = GetSnapGridSize();
-// 					AnchorNodeNewPos.X = SnapSize * FMath::Round(AnchorNodeNewPos.X/SnapSize);
-// 					AnchorNodeNewPos.Y = SnapSize * FMath::Round(AnchorNodeNewPos.Y/SnapSize);
+// 					AnchorNodeNewPos.X = SnapSize * FMath::RoundToInt(AnchorNodeNewPos.X/SnapSize);
+// 					AnchorNodeNewPos.Y = SnapSize * FMath::RoundToInt(AnchorNodeNewPos.Y/SnapSize);
 // 
 // 					// Dragging an unselected node automatically selects it.
 // 					SelectionManager.StartDraggingNode(NodeBeingDragged->GetObjectBeingDisplayed(), MouseEvent);
@@ -416,7 +416,7 @@ FReply SPaperEditorViewport::OnMouseWheel(const FGeometry& MyGeometry, const FPo
 	FVector2D PointToMaintainGraphSpace = PanelCoordToGraphCoord( WidgetSpaceCursorPos );
 
 
-	const int32 ZoomLevelDelta = FMath::Floor(MouseEvent.GetWheelDelta());
+	const int32 ZoomLevelDelta = FMath::FloorToInt(MouseEvent.GetWheelDelta());
 
 	const bool bAllowFullZoomRange = true;
 /*
diff --git a/Engine/Plugins/Experimental/Paper2D/Source/Paper2DEditor/Private/SpriteEditor/SpriteEditorViewportClient.cpp b/Engine/Plugins/Experimental/Paper2D/Source/Paper2DEditor/Private/SpriteEditor/SpriteEditorViewportClient.cpp
--- a/Engine/Plugins/Experimental/Paper2D/Source/Paper2DEditor/Private/SpriteEditor/SpriteEditorViewportClient.cpp
+++ b/Engine/Plugins/Experimental/Paper2D/Source/Paper2DEditor/Private/SpriteEditor/SpriteEditorViewportClient.cpp
@@ -336,7 +336,7 @@ void FSpriteEditorViewportClient::DrawSocketNames(FViewport& InViewport, FSceneV
 // 				// Figure out the text height
 // 				FTextSizingParameters Parameters(GEngine->GetSmallFont(), 1.0f, 1.0f);
 // 				UCanvas::CanvasStringSize(Parameters, *Socket->SocketName.ToString());
-// 				int32 YL = FMath::Trunc(Parameters.DrawYL);
+// 				int32 YL = FMath::TruncToInt(Parameters.DrawYL);
 // 
 // 				DrawAngles(&Canvas, XPos, YPos + YL, 
 // 					Widget->GetCurrentAxis(), 
@@ -737,8 +737,8 @@ void FSpriteEditorViewportClient::SnapAllVerticesToPixelGrid()
 			{
 				FVector2D& Vertex = Poly.Vertices[VertexIndex];
 
-				Vertex.X = FMath::Round(Vertex.X);
-				Vertex.Y = FMath::Round(Vertex.Y);
+				Vertex.X = FMath::RoundToInt(Vertex.X);
+				Vertex.Y = FMath::RoundToInt(Vertex.Y);
 				bManipulationDirtiedSomething = true;
 			}
 		}
diff --git a/Engine/Plugins/Experimental/Paper2D/Source/Paper2DEditor/Private/TileMapEditing/EdModeTileMap.cpp b/Engine/Plugins/Experimental/Paper2D/Source/Paper2DEditor/Private/TileMapEditing/EdModeTileMap.cpp
--- a/Engine/Plugins/Experimental/Paper2D/Source/Paper2DEditor/Private/TileMapEditing/EdModeTileMap.cpp
+++ b/Engine/Plugins/Experimental/Paper2D/Source/Paper2DEditor/Private/TileMapEditing/EdModeTileMap.cpp
@@ -349,8 +349,8 @@ UPaperTileLayer* FEdModeTileMap::GetSelectedLayerUnderCursor(const FViewportCurs
 					const float NormalizedX = (Intersection.X + 0.5f * TileMap->TileWidth) / WX;
 					const float NormalizedY = (-Intersection.Z + 0.5f * TileMap->TileHeight) / WY;
 
-					OutTileX = FMath::Floor(NormalizedX * TileMap->MapWidth);
-					OutTileY = FMath::Floor(NormalizedY * TileMap->MapHeight);
+					OutTileX = FMath::FloorToInt(NormalizedX * TileMap->MapWidth);
+					OutTileY = FMath::FloorToInt(NormalizedY * TileMap->MapHeight);
 					
 					if ((OutTileX > -BrushWidth) && (OutTileX < TileMap->MapWidth) && (OutTileY > -BrushHeight) && (OutTileY < TileMap->MapHeight))
 						//(NormalizedX >= 0.0f) && (NormalizedX < 1.0f) && (NormalizedY >= 0.0f) && (NormalizedY < 1.0f))
diff --git a/Engine/Plugins/Experimental/Phya/Source/Phya/Private/PhyaCollisionHandler.cpp b/Engine/Plugins/Experimental/Phya/Source/Phya/Private/PhyaCollisionHandler.cpp
--- a/Engine/Plugins/Experimental/Phya/Source/Phya/Private/PhyaCollisionHandler.cpp
+++ b/Engine/Plugins/Experimental/Phya/Source/Phya/Private/PhyaCollisionHandler.cpp
@@ -134,7 +134,7 @@ void UPhyaCollisionHandler::StreamingWaveUnderflow(USoundWaveStreaming* InStream
 
 	const int32 QueuedSamples = StreamingWave->GetAvailableAudioByteCount()/sizeof(uint16);
 	const int32 SamplesNeeded = SamplesRequired - QueuedSamples;
-	const int32 BlocksNeeded = FMath::Ceil(SamplesNeeded/BlockSize);
+	const int32 BlocksNeeded = FMath::CeilToInt(SamplesNeeded/BlockSize);
 
 	UE_LOG(LogTemp, Log, TEXT("Creating %d blocks for %s"), BlocksNeeded, *GetWorld()->GetPathName());
 
diff --git a/Engine/Plugins/OculusRift/Source/OculusRift/Private/OculusRiftHMD.cpp b/Engine/Plugins/OculusRift/Source/OculusRift/Private/OculusRiftHMD.cpp
--- a/Engine/Plugins/OculusRift/Source/OculusRift/Private/OculusRiftHMD.cpp
+++ b/Engine/Plugins/OculusRift/Source/OculusRift/Private/OculusRiftHMD.cpp
@@ -1345,7 +1345,7 @@ void FOculusRiftHMD::PushViewportCanvas(EStereoscopicPass StereoPass, FCanvas *I
 {
 	if (StereoPass != eSSP_FULL)
 	{
-		int32 SideSizeX = FMath::Trunc(InViewport->GetSizeXY().X * 0.5);
+		int32 SideSizeX = FMath::TruncToInt(InViewport->GetSizeXY().X * 0.5);
 
 		// !AB: temporarily assuming all canvases are at Z = 1.0f and calculating
 		// stereo disparity right here. Stereo disparity should be calculated for each
diff --git a/Engine/Source/Developer/CrashTracker/Private/CrashVideoCapture.cpp b/Engine/Source/Developer/CrashTracker/Private/CrashVideoCapture.cpp
--- a/Engine/Source/Developer/CrashTracker/Private/CrashVideoCapture.cpp
+++ b/Engine/Source/Developer/CrashTracker/Private/CrashVideoCapture.cpp
@@ -220,14 +220,14 @@ void FCrashVideoCapture::SaveFinalCapture( const FString& CapturePath )
 	// this may not be useful with a crash on a non-game thread, because the final frame can't be drawn
 	if (FinalFrame)
 	{
-		int32 ExtraFramesAdded = FMath::Trunc(1.f / CrashTrackerConstants::CaptureFrequency);
+		int32 ExtraFramesAdded = FMath::TruncToInt(1.f / CrashTrackerConstants::CaptureFrequency);
 		for (int32 i = 0; i < ExtraFramesAdded; ++i)
 		{
 			OrderedDataFrames.Add(FinalFrame);
 		}
 	}
 
-	DumpOutMJPEGAVI(OrderedDataFrames, CapturePath, Width, Height, FMath::Trunc(1.f / CrashTrackerConstants::CaptureFrequency));
+	DumpOutMJPEGAVI(OrderedDataFrames, CapturePath, Width, Height, FMath::TruncToInt(1.f / CrashTrackerConstants::CaptureFrequency));
 
 	UE_LOG(LogCrashTracker, Log, TEXT("Finished dumping out crash capture video. %i Frames at %ix%i"), OrderedDataFrames.Num(), Width, Height);
 }
@@ -263,7 +263,7 @@ void FCrashVideoCapture::Update(float DeltaSeconds)
 		CurrentAccumSeconds += DeltaSeconds;
 		if (CurrentAccumSeconds > CrashTrackerConstants::CaptureFrequency)
 		{
-			CurrentAccumSeconds -= FMath::Trunc(CurrentAccumSeconds / CrashTrackerConstants::CaptureFrequency) * CrashTrackerConstants::CaptureFrequency;
+			CurrentAccumSeconds -= FMath::TruncToFloat(CurrentAccumSeconds / CrashTrackerConstants::CaptureFrequency) * CrashTrackerConstants::CaptureFrequency;
 			bShouldUpdate = true;
 		}
 
diff --git a/Engine/Source/Developer/ImageCore/Private/ImageCore.cpp b/Engine/Source/Developer/ImageCore/Private/ImageCore.cpp
--- a/Engine/Source/Developer/ImageCore/Private/ImageCore.cpp
+++ b/Engine/Source/Developer/ImageCore/Private/ImageCore.cpp
@@ -87,10 +87,10 @@ static void CopyImage(const FImage& SrcImage, FImage& DestImage)
 				for (int32 TexelIndex = 0; TexelIndex < NumTexels; ++TexelIndex)
 				{
 					int32 DestIndex = TexelIndex * 4;
-					DestColors[DestIndex + 0] = FMath::Clamp(FMath::Floor(SrcColors[TexelIndex].R * 65535.999f), 0, 65535);
-					DestColors[DestIndex + 1] = FMath::Clamp(FMath::Floor(SrcColors[TexelIndex].G * 65535.999f), 0, 65535);
-					DestColors[DestIndex + 2] = FMath::Clamp(FMath::Floor(SrcColors[TexelIndex].B * 65535.999f), 0, 65535);
-					DestColors[DestIndex + 3] = FMath::Clamp(FMath::Floor(SrcColors[TexelIndex].A * 65535.999f), 0, 65535);
+					DestColors[DestIndex + 0] = FMath::Clamp(FMath::FloorToInt(SrcColors[TexelIndex].R * 65535.999f), 0, 65535);
+					DestColors[DestIndex + 1] = FMath::Clamp(FMath::FloorToInt(SrcColors[TexelIndex].G * 65535.999f), 0, 65535);
+					DestColors[DestIndex + 2] = FMath::Clamp(FMath::FloorToInt(SrcColors[TexelIndex].B * 65535.999f), 0, 65535);
+					DestColors[DestIndex + 3] = FMath::Clamp(FMath::FloorToInt(SrcColors[TexelIndex].A * 65535.999f), 0, 65535);
 				}
 			}
 			break;
diff --git a/Engine/Source/Developer/MeshSimplifier/Private/QuadricMeshReduction.cpp b/Engine/Source/Developer/MeshSimplifier/Private/QuadricMeshReduction.cpp
--- a/Engine/Source/Developer/MeshSimplifier/Private/QuadricMeshReduction.cpp
+++ b/Engine/Source/Developer/MeshSimplifier/Private/QuadricMeshReduction.cpp
@@ -285,9 +285,9 @@ public:
 
 		for( uint32 i = 0; i < NumVerts; i++ )
 		{
-			Verts[i].Position.x = 0.0625f * FMath::Round( 16.0f * Verts[i].Position.X );
-			Verts[i].Position.y = 0.0625f * FMath::Round( 16.0f * Verts[i].Position.Y );
-			Verts[i].Position.z = 0.0625f * FMath::Round( 16.0f * Verts[i].Position.Z );
+			Verts[i].Position.x = 0.0625f * FMath::RoundToFloat( 16.0f * Verts[i].Position.X );
+			Verts[i].Position.y = 0.0625f * FMath::RoundToFloat( 16.0f * Verts[i].Position.Y );
+			Verts[i].Position.z = 0.0625f * FMath::RoundToFloat( 16.0f * Verts[i].Position.Z );
 		}
 
 		for( uint32 i = 0; i < NumVerts; i++ )
diff --git a/Engine/Source/Developer/MeshUtilities/Private/MeshUtilities.cpp b/Engine/Source/Developer/MeshUtilities/Private/MeshUtilities.cpp
--- a/Engine/Source/Developer/MeshUtilities/Private/MeshUtilities.cpp
+++ b/Engine/Source/Developer/MeshUtilities/Private/MeshUtilities.cpp
@@ -1498,7 +1498,7 @@ static void ComputeStreamingTextureFactors(
 			// Disregard upper 75% of texel ratios.
 			// This is to ignore backfacing surfaces or other non-visible surfaces that tend to map a small number of texels to a large surface.
 			TexelRatios[UVIndex].Sort( TGreater<float>() );
-			float TexelRatio = TexelRatios[UVIndex][ FMath::Trunc(TexelRatios[UVIndex].Num() * 0.75f) ];
+			float TexelRatio = TexelRatios[UVIndex][ FMath::TruncToInt(TexelRatios[UVIndex].Num() * 0.75f) ];
 			OutStreamingTextureFactors[UVIndex] = TexelRatio;
 		}
 	}
@@ -2900,7 +2900,7 @@ public:
 		SortedLightmaps.Sort([](TPair<int32, uint32> L, TPair<int32, uint32> R) { return R.Value < L.Value; });
 		
 		// Try to pack, increasing atlas resolution with each step
-		uint32 PackedSize = FMath::RoundUpToPowerOfTwo(FMath::Round(FMath::Sqrt(TotalArea)));
+		uint32 PackedSize = FMath::RoundUpToPowerOfTwo(FMath::RoundToInt(FMath::Sqrt(TotalArea)));
 		for (int32 i = 0; i < 10; ++i) // 2 iterations should be enough >.<
 		{
 			PackedLightmapAtlas = new FLightmapAtlas(PackedSize);
diff --git a/Engine/Source/Developer/MovieSceneTools/Private/AudioTrackEditor.cpp b/Engine/Source/Developer/MovieSceneTools/Private/AudioTrackEditor.cpp
--- a/Engine/Source/Developer/MovieSceneTools/Private/AudioTrackEditor.cpp
+++ b/Engine/Source/Developer/MovieSceneTools/Private/AudioTrackEditor.cpp
@@ -192,11 +192,11 @@ void FAudioThumbnail::GenerateWaveformPreview(TArray<uint8>& OutData, TRange<flo
 	{
 		float LookupTime = ((float)(X - 0.5f) / (float)GetSize().X) * DrawRangeSize + DrawRange.GetLowerBoundValue();
 		float LookupFraction = (LookupTime - AudioTrueRange.GetLowerBoundValue()) / TrueRangeSize;
-		int32 LookupIndex = FMath::Trunc(LookupFraction * LookupSize);
+		int32 LookupIndex = FMath::TruncToInt(LookupFraction * LookupSize);
 		
 		float NextLookupTime = ((float)(X + 0.5f) / (float)GetSize().X) * DrawRangeSize + DrawRange.GetLowerBoundValue();
 		float NextLookupFraction = (NextLookupTime - AudioTrueRange.GetLowerBoundValue()) / TrueRangeSize;
-		int32 NextLookupIndex = FMath::Trunc(NextLookupFraction * LookupSize);
+		int32 NextLookupIndex = FMath::TruncToInt(NextLookupFraction * LookupSize);
 
 		DrawWaveformLine(OutData, SoundWave->NumChannels, LookupData, LookupIndex, NextLookupIndex, LookupSize, X, GetSize().Y / 2, GetSize().Y / 2);
 	}
@@ -242,7 +242,7 @@ void FAudioThumbnail::DrawWaveformLine(TArray<uint8>& OutData, int32 NumChannels
 		for (int32 ChannelIndex = 0; ChannelIndex < NumChannels; ++ChannelIndex)
 		{
 			int32 DataPoint = LookupData[i + ChannelIndex];
-			int32 HistogramData = FMath::Clamp(FMath::Trunc(FMath::Abs(DataPoint) / 32768.f * MaxAmplitude), 0, MaxAmplitude - 1);
+			int32 HistogramData = FMath::Clamp(FMath::TruncToInt(FMath::Abs(DataPoint) / 32768.f * MaxAmplitude), 0, MaxAmplitude - 1);
 
 			++ChannelData[ChannelIndex].Histogram[HistogramData];
 			++ChannelData[ChannelIndex].HistogramSum;
@@ -284,14 +284,14 @@ void FAudioThumbnail::DrawWaveformLine(TArray<uint8>& OutData, int32 NumChannels
 	}
 	float RootMeanSquare = FMath::Sqrt(SummedData / (float)Samples);
 	
-	int32 RMSHeight = FMath::Clamp(FMath::Trunc(RootMeanSquare / 32768.f * MaxAmplitude), 1, MaxAmplitude);
-	int32 MaxHeight = FMath::Clamp(FMath::Trunc(DataAmplitude / 32768.f * MaxAmplitude), 1, MaxAmplitude);
+	int32 RMSHeight = FMath::Clamp(FMath::TruncToInt(RootMeanSquare / 32768.f * MaxAmplitude), 1, MaxAmplitude);
+	int32 MaxHeight = FMath::Clamp(FMath::TruncToInt(DataAmplitude / 32768.f * MaxAmplitude), 1, MaxAmplitude);
 	float HeightRange = FMath::Max(1, MaxHeight - RMSHeight);
 	
 	for (int32 i = 0; i < MaxHeight; ++i)
 	{
 		float Intensity = (MaxHeight - i) / HeightRange;
-		int32 IntegralIntensity = FMath::Clamp(FMath::Trunc(Intensity * 255), 0, 255);
+		int32 IntegralIntensity = FMath::Clamp(FMath::TruncToInt(Intensity * 255), 0, 255);
 
 		int32 Y = XAxisHeight - i;
 		int32 Index = (Y * GetSize().X + X) * GPixelFormats[PF_B8G8R8A8].BlockBytes;
diff --git a/Engine/Source/Developer/Profiler/Private/ProfilerDataSource.h b/Engine/Source/Developer/Profiler/Private/ProfilerDataSource.h
--- a/Engine/Source/Developer/Profiler/Private/ProfilerDataSource.h
+++ b/Engine/Source/Developer/Profiler/Private/ProfilerDataSource.h
@@ -246,7 +246,7 @@ public:
 	FORCEINLINE const uint32 GetStartIndexFromTimeRange( const float StartTimeMS, const float EndTimeMS ) const
 	{
 		CheckInvariants( StartTimeMS, EndTimeMS );
-		const uint32 Index = FMath::Trunc( StartTimeMS * InvTimeAccuracyMS );
+		const uint32 Index = FMath::TruncToInt( StartTimeMS * InvTimeAccuracyMS );
 		return Index;
 	}
 
@@ -267,7 +267,7 @@ public:
 		Type Result = (Type)0;
 		const uint32 Index = GetStartIndexFromTimeRange( StartTimeMS, EndTimeMS );
 		const uint32 CurrentChunkIndex = Index / ThisCacheDataContainer::NumElementsPerChunk;
-		const uint32 TotalNumFrames = FMath::Trunc( static_cast<const ManagerClass*>(this)->GetTotalTimeMS() * InvTimeAccuracyMS );
+		const uint32 TotalNumFrames = FMath::TruncToInt( static_cast<const ManagerClass*>(this)->GetTotalTimeMS() * InvTimeAccuracyMS );
 
 		const uint32 NumNeededChunks = (TotalNumFrames + ThisCacheDataContainer::NumElementsPerChunk - 1) / ThisCacheDataContainer::NumElementsPerChunk;
 		const uint32 NumMissingValues = TotalNumFrames - ThisCacheDataContainer::CachedValues.Num();
@@ -627,7 +627,7 @@ public:
 	{	
 		if( GetSourcesNum() > 0 )
 		{
-			return FMath::Trunc( GetTotalTimeMS() * InvTimeAccuracyMS );
+			return FMath::TruncToInt( GetTotalTimeMS() * InvTimeAccuracyMS );
 		}
 		else
 		{
diff --git a/Engine/Source/Developer/Profiler/Private/ProfilerFPSAnalyzer.cpp b/Engine/Source/Developer/Profiler/Private/ProfilerFPSAnalyzer.cpp
--- a/Engine/Source/Developer/Profiler/Private/ProfilerFPSAnalyzer.cpp
+++ b/Engine/Source/Developer/Profiler/Private/ProfilerFPSAnalyzer.cpp
@@ -24,7 +24,7 @@ void FFPSAnalyzer::AddSample( float FPSSample )
 	Samples.Add(FPSSample);
 
 	float DeltaSeconds = 1.0f / FPSSample;
-	int32 Index = FPlatformMath::Floor(FPSSample / Interval);
+	int32 Index = FPlatformMath::FloorToInt(FPSSample / Interval);
 	Index = Index >= Histogram.Num() ? Histogram.Num()-1 : Index;
 	Histogram[Index].Count++;
 	Histogram[Index].CummulativeTime += DeltaSeconds;
@@ -59,7 +59,7 @@ int32 FFPSAnalyzer::GetTotalCount()
 
 int32 FFPSAnalyzer::GetCount(float InMinVal, float InMaxVal)
 {
-	int32 Index = FPlatformMath::Floor(InMinVal / Interval);
+	int32 Index = FPlatformMath::FloorToInt(InMinVal / Interval);
 	Index = Index >= Histogram.Num() ? Histogram.Num()-1 : Index;
 	return Histogram[Index].Count;
 }
\ No newline at end of file
diff --git a/Engine/Source/Developer/Profiler/Private/Widgets/SDataGraph.cpp b/Engine/Source/Developer/Profiler/Private/Widgets/SDataGraph.cpp
--- a/Engine/Source/Developer/Profiler/Private/Widgets/SDataGraph.cpp
+++ b/Engine/Source/Developer/Profiler/Private/Widgets/SDataGraph.cpp
@@ -263,7 +263,7 @@ void SDataGraph::UpdateState()
 			NumDataPoints = (int32)GraphDesc->CombinedGraphDataSource->GetNumFrames();
 		}
 
-		NumVisiblePoints = FMath::Max( 0, FMath::Trunc(ThisGeometry.Size.X) / DistanceBetweenPoints );
+		NumVisiblePoints = FMath::Max( 0, FMath::TruncToInt(ThisGeometry.Size.X) / DistanceBetweenPoints );
 		// GraphOffset - Updated by OnMouseMove or by ScrollTo
 		GraphOffset = FMath::Clamp( GraphOffset, 0, FMath::Max(NumDataPoints-NumVisiblePoints,0) );
 		
@@ -587,7 +587,7 @@ int32 SDataGraph::OnPaint
 			{
 				const float MarkerPosX = (FrameIndex - GraphOffset) * DistanceBetweenPoints;
 				const float ElapsedFrameTimeMS = FrameIndex * FTimeAccuracy::AsFrameTime( TimeBasedAccuracy );
-				const int32 ElapsedFrameTime = FMath::Max( FMath::Round( ElapsedFrameTimeMS * 0.001f ) - 1, 0 );
+				const int32 ElapsedFrameTime = FMath::Max( FMath::RoundToInt( ElapsedFrameTimeMS * 0.001f ) - 1, 0 );
 				const int32 AccumulatedFrameCounter = bCanBeDisplayedAsMulti ? FrameIndex : DataProvider->GetAccumulatedFrameCounter(ElapsedFrameTime);
 
 				LinePoints.Add( FVector2D(MarkerPosX, 0.0) );
@@ -1144,7 +1144,7 @@ FReply SDataGraph::OnMouseMove( const FGeometry& MyGeometry, const FPointerEvent
 			const float ScrollByAmount = -MouseEvent.GetCursorDelta().X * (1.0f/DistanceBetweenPoints);
 			RealGraphOffset += ScrollByAmount;
 
-			GraphOffset = FMath::Clamp( FMath::Trunc( RealGraphOffset ), 0, FMath::Max(NumDataPoints-NumVisiblePoints,0) );
+			GraphOffset = FMath::Clamp( FMath::TruncToInt( RealGraphOffset ), 0, FMath::Max(NumDataPoints-NumVisiblePoints,0) );
 			OnGraphOffsetChanged.ExecuteIfBound( GraphOffset );
 
 			Reply = FReply::Handled();
@@ -1157,7 +1157,7 @@ FReply SDataGraph::OnMouseMove( const FGeometry& MyGeometry, const FPointerEvent
 const int32 SDataGraph::CalculateFrameIndex( const FVector2D InMousePosition ) const
 {
 	const float ScaleX = 1.0f/DistanceBetweenPoints;
-	const int32 MousePositionOffset = FMath::Trunc( (InMousePosition.X+HalfGraphMarkerWidth) * ScaleX );
+	const int32 MousePositionOffset = FMath::TruncToInt( (InMousePosition.X+HalfGraphMarkerWidth) * ScaleX );
 	return FMath::Clamp( GraphOffset+MousePositionOffset, 0, NumDataPoints-1 );
 }
 
diff --git a/Engine/Source/Developer/Profiler/Private/Widgets/SHistogram.h b/Engine/Source/Developer/Profiler/Private/Widgets/SHistogram.h
--- a/Engine/Source/Developer/Profiler/Private/Widgets/SHistogram.h
+++ b/Engine/Source/Developer/Profiler/Private/Widgets/SHistogram.h
@@ -44,7 +44,7 @@ public:
 		, MaxValue(InMaxValue)
 		, Normalize(InBinNormalize)
 	{
-		BinCount = FPlatformMath::Ceil( (MaxValue - MinValue) / Interval )+1;	// +1 for data beyond the max
+		BinCount = FPlatformMath::CeilToInt( (MaxValue - MinValue) / Interval )+1;	// +1 for data beyond the max
 	}
 
 	/** Retrieves the bin count */
diff --git a/Engine/Source/Developer/Profiler/Private/Widgets/SProfilerGraphPanel.cpp b/Engine/Source/Developer/Profiler/Private/Widgets/SProfilerGraphPanel.cpp
--- a/Engine/Source/Developer/Profiler/Private/Widgets/SProfilerGraphPanel.cpp
+++ b/Engine/Source/Developer/Profiler/Private/Widgets/SProfilerGraphPanel.cpp
@@ -94,7 +94,7 @@ void SProfilerGraphPanel::HorizontalScrollBar_OnUserScrolled( float ScrollOffset
  
 	HorizontalScrollBar->SetState( ScrollOffset, ThumbSizeFraction );
  
-	GraphOffset = FMath::Trunc( ScrollOffset * NumDataPoints );
+	GraphOffset = FMath::TruncToInt( ScrollOffset * NumDataPoints );
 	DataGraph->ScrollTo( GraphOffset );
 
 	FProfilerManager::Get()->GetProfilerWindow()->ProfilerMiniView->OnSelectionBoxChanged( GraphOffset, GraphOffset + NumVisiblePoints );
diff --git a/Engine/Source/Developer/Profiler/Private/Widgets/SProfilerMiniView.cpp b/Engine/Source/Developer/Profiler/Private/Widgets/SProfilerMiniView.cpp
--- a/Engine/Source/Developer/Profiler/Private/Widgets/SProfilerMiniView.cpp
+++ b/Engine/Source/Developer/Profiler/Private/Widgets/SProfilerMiniView.cpp
@@ -129,13 +129,13 @@ int32 SProfilerMiniView::OnPaint( const FGeometry& AllottedGeometry, const FSlat
 		float NextSamplePosX = GetNumPixelsPerSample();
 		for( const FMiniViewSample& MiniViewSample : MiniViewSamples )
 		{
-			const float AllSizeY = FMath::TruncFloat( MiniViewSample.TotalThreadTime*SampleScaleY );
-			const float GTSizeY = FMath::TruncFloat( MiniViewSample.GameThreadTime*SampleScaleY );
-			const float RTSizeY = FMath::TruncFloat( MiniViewSample.RenderThreadTime*SampleScaleY );
+			const float AllSizeY = FMath::TruncToFloat( MiniViewSample.TotalThreadTime*SampleScaleY );
+			const float GTSizeY = FMath::TruncToFloat( MiniViewSample.GameThreadTime*SampleScaleY );
+			const float RTSizeY = FMath::TruncToFloat( MiniViewSample.RenderThreadTime*SampleScaleY );
 			const float OtherSizeY = AllSizeY - GTSizeY - RTSizeY;
 
-			const float DestSamplePosX0 = FMath::TruncFloat( CurrentSamplePosX );
-			const float DestSamplePosX1 = FMath::TruncFloat( NextSamplePosX );
+			const float DestSamplePosX0 = FMath::TruncToFloat( CurrentSamplePosX );
+			const float DestSamplePosX1 = FMath::TruncToFloat( NextSamplePosX );
 			const float DestSampleSizeX = DestSamplePosX1 - DestSamplePosX0;
 
 			// The game thread on the bottom.
@@ -186,8 +186,8 @@ int32 SProfilerMiniView::OnPaint( const FGeometry& AllottedGeometry, const FSlat
 		
 		const int32 MaxFrameIndex = AllFrames.Num() - 1;
 		
-		const int32 SelectionBoxX0 = FMath::Trunc( FrameIndexToPosition( SelectionBoxFrameStart ) );
-		const int32 SelectionBoxX1 = FMath::Trunc( FrameIndexToPosition( SelectionBoxFrameEnd + 1 ) );
+		const int32 SelectionBoxX0 = FMath::TruncToInt( FrameIndexToPosition( SelectionBoxFrameStart ) );
+		const int32 SelectionBoxX1 = FMath::TruncToInt(FrameIndexToPosition(SelectionBoxFrameEnd + 1));
 
 		if( SelectionBoxFrameStart > 0 )
 		{
@@ -617,7 +617,7 @@ const int32 SProfilerMiniView::PositionToFrameIndex( const float InPositionX ) c
 		const float ScaleRatio = (float)NumAllFrames / (float)NumMiniViewSamples;
 		const float MouseSampleIndex = InPositionX / GetNumPixelsPerSample();
 
-		FrameIndex = FMath::Trunc( MouseSampleIndex * ScaleRatio );
+		FrameIndex = FMath::TruncToInt(MouseSampleIndex * ScaleRatio);
 		FrameIndex = FMath::Clamp( FrameIndex, 0, NumAllFrames - 1 );
 	}
 
@@ -639,7 +639,7 @@ void SProfilerMiniView::ProcessData()
 
 	UpdateNumPixelsPerSample();
 
-	const int32 NumMiniViewSamples = FMath::Trunc( ThisGeometry.Size.X / GetNumPixelsPerSample() );
+	const int32 NumMiniViewSamples = FMath::TruncToInt(ThisGeometry.Size.X / GetNumPixelsPerSample());
 	MiniViewSamples.Reset( NumMiniViewSamples );
 
 	MiniViewSamples.AddZeroed( NumMiniViewSamples );
@@ -657,7 +657,7 @@ void SProfilerMiniView::ProcessData()
 		for( int32 FrameIndex = 0; FrameIndex < NumAllFrames; ++FrameIndex )
 		{
 			const FFrameThreadTimes& FrameThreadTimes = AllFrames[FrameIndex];
-			const int32 SampleIndex = FMath::Trunc( CurrentSample );
+			const int32 SampleIndex = FMath::TruncToInt(CurrentSample);
 
 			FMiniViewSample& Dest = MiniViewSamples[SampleIndex];
 			Dest.AddFrameAndFindMax( FrameThreadTimes );
diff --git a/Engine/Source/Developer/RealtimeProfiler/Private/SRealtimeProfilerLineGraph.cpp b/Engine/Source/Developer/RealtimeProfiler/Private/SRealtimeProfilerLineGraph.cpp
--- a/Engine/Source/Developer/RealtimeProfiler/Private/SRealtimeProfilerLineGraph.cpp
+++ b/Engine/Source/Developer/RealtimeProfiler/Private/SRealtimeProfilerLineGraph.cpp
@@ -166,7 +166,7 @@ int32 SRealtimeProfilerLineGraph::OnPaint( const FGeometry& AllottedGeometry, co
 	//Convert Offset to FrameOffset
 	float FrameOffset = (-Offset/Zoom)*MaxFrames.Get();
 
-	int32 StartPointIndex = FMath::Floor( FrameOffset );
+	int32 StartPointIndex = FMath::FloorToInt( FrameOffset );
 
 	if(StartPointIndex <= 0)
 		StartPointIndex = 0;
diff --git a/Engine/Source/Developer/StandaloneRenderer/Private/Linux/OpenGL/SlateOpenGLViewport.cpp b/Engine/Source/Developer/StandaloneRenderer/Private/Linux/OpenGL/SlateOpenGLViewport.cpp
--- a/Engine/Source/Developer/StandaloneRenderer/Private/Linux/OpenGL/SlateOpenGLViewport.cpp
+++ b/Engine/Source/Developer/StandaloneRenderer/Private/Linux/OpenGL/SlateOpenGLViewport.cpp
@@ -15,8 +15,8 @@ void FSlateOpenGLViewport::Initialize( TSharedRef<SWindow> InWindow, const FSlat
 	RenderingContext.Initialize(NativeWindow->GetOSWindowHandle(), &SharedContext);
 
 	// Create an OpenGL viewport
-	const int32 Width = FMath::Trunc(InWindow->GetSizeInScreen().X);
-	const int32 Height = FMath::Trunc(InWindow->GetSizeInScreen().Y);
+	const int32 Width = FMath::TruncToInt(InWindow->GetSizeInScreen().X);
+	const int32 Height = FMath::TruncToInt(InWindow->GetSizeInScreen().Y);
 
 	ProjectionMatrix = CreateProjectionMatrix( Width, Height );
 
diff --git a/Engine/Source/Developer/StandaloneRenderer/Private/Mac/OpenGL/SlateOpenGLViewport.cpp b/Engine/Source/Developer/StandaloneRenderer/Private/Mac/OpenGL/SlateOpenGLViewport.cpp
--- a/Engine/Source/Developer/StandaloneRenderer/Private/Mac/OpenGL/SlateOpenGLViewport.cpp
+++ b/Engine/Source/Developer/StandaloneRenderer/Private/Mac/OpenGL/SlateOpenGLViewport.cpp
@@ -16,8 +16,8 @@ void FSlateOpenGLViewport::Initialize( TSharedRef<SWindow> InWindow, const FSlat
 	TSharedRef<FGenericWindow> NativeWindow = InWindow->GetNativeWindow().ToSharedRef();
 	RenderingContext.Initialize( NativeWindow->GetOSWindowHandle(), &SharedContext );
 
-	const int32 Width = FMath::Trunc(InWindow->GetSizeInScreen().X);
-	const int32 Height = FMath::Trunc(InWindow->GetSizeInScreen().Y);
+	const int32 Width = FMath::TruncToInt(InWindow->GetSizeInScreen().X);
+	const int32 Height = FMath::TruncToInt(InWindow->GetSizeInScreen().Y);
 
 	ProjectionMatrix = CreateProjectionMatrix( Width, Height );
 
diff --git a/Engine/Source/Developer/StandaloneRenderer/Private/OpenGL/SlateOpenGLRenderer.cpp b/Engine/Source/Developer/StandaloneRenderer/Private/OpenGL/SlateOpenGLRenderer.cpp
--- a/Engine/Source/Developer/StandaloneRenderer/Private/OpenGL/SlateOpenGLRenderer.cpp
+++ b/Engine/Source/Developer/StandaloneRenderer/Private/OpenGL/SlateOpenGLRenderer.cpp
@@ -179,8 +179,8 @@ void FSlateOpenGLRenderer::RequestResize( const TSharedPtr<SWindow>& InWindow, u
 
 void FSlateOpenGLRenderer::Private_ResizeViewport( const FVector2D& WindowSize, FSlateOpenGLViewport& InViewport, bool bFullscreen )
 {
-	uint32 Width = FMath::Trunc(WindowSize.X);
-	uint32 Height = FMath::Trunc(WindowSize.Y);
+	uint32 Width = FMath::TruncToInt(WindowSize.X);
+	uint32 Height = FMath::TruncToInt(WindowSize.Y);
 
 	InViewport.Resize( Width, Height, bFullscreen );
 }
diff --git a/Engine/Source/Developer/StandaloneRenderer/Private/Windows/D3D/SlateD3DRenderer.cpp b/Engine/Source/Developer/StandaloneRenderer/Private/Windows/D3D/SlateD3DRenderer.cpp
--- a/Engine/Source/Developer/StandaloneRenderer/Private/Windows/D3D/SlateD3DRenderer.cpp
+++ b/Engine/Source/Developer/StandaloneRenderer/Private/Windows/D3D/SlateD3DRenderer.cpp
@@ -130,8 +130,8 @@ void FSlateD3DRenderer::Private_CreateViewport( TSharedRef<SWindow> InWindow, co
 	DXGI_SWAP_CHAIN_DESC SwapChainDesc;
 	FMemory::Memzero(&SwapChainDesc, sizeof(SwapChainDesc) );
 	SwapChainDesc.BufferCount = 1;
-	SwapChainDesc.BufferDesc.Width = FMath::Trunc(WindowSize.X);
-	SwapChainDesc.BufferDesc.Height = FMath::Trunc(WindowSize.Y);
+	SwapChainDesc.BufferDesc.Width = FMath::TruncToInt(WindowSize.X);
+	SwapChainDesc.BufferDesc.Height = FMath::TruncToInt(WindowSize.Y);
 	SwapChainDesc.BufferDesc.Format = DXGI_FORMAT_B8G8R8A8_UNORM;
 	SwapChainDesc.BufferDesc.RefreshRate.Numerator = 0;
 	SwapChainDesc.BufferDesc.RefreshRate.Denominator = 1;
@@ -160,8 +160,8 @@ void FSlateD3DRenderer::Private_CreateViewport( TSharedRef<SWindow> InWindow, co
 	Hr = DXGIFactory->CreateSwapChain(DXGIDevice.GetReference(), &SwapChainDesc, Viewport.D3DSwapChain.GetInitReference() );
 	check( SUCCEEDED(Hr) );
 
-	uint32 Width = FMath::Trunc(WindowSize.X);
-	uint32 Height = FMath::Trunc(WindowSize.Y);
+	uint32 Width = FMath::TruncToInt(WindowSize.X);
+	uint32 Height = FMath::TruncToInt(WindowSize.Y);
 
 	Viewport.ViewportInfo.MaxDepth = 1.0f;
 	Viewport.ViewportInfo.MinDepth = 0.0f;
diff --git a/Engine/Source/Developer/StandaloneRenderer/Private/ios/SlateOpenGLViewport.cpp b/Engine/Source/Developer/StandaloneRenderer/Private/ios/SlateOpenGLViewport.cpp
--- a/Engine/Source/Developer/StandaloneRenderer/Private/ios/SlateOpenGLViewport.cpp
+++ b/Engine/Source/Developer/StandaloneRenderer/Private/ios/SlateOpenGLViewport.cpp
@@ -22,8 +22,8 @@ void FSlateOpenGLViewport::Initialize( TSharedRef<SWindow> InWindow, const FSlat
 
 	FSlateRect FullScreenSize = InWindow->GetFullScreenInfo();
 
-	const int Width = FMath::Trunc(FullScreenSize.Right - FullScreenSize.Left);
-	const int Height = FMath::Trunc(FullScreenSize.Bottom - FullScreenSize.Top);
+	const int Width = FMath::TruncToInt(FullScreenSize.Right - FullScreenSize.Left);
+	const int Height = FMath::TruncToInt(FullScreenSize.Bottom - FullScreenSize.Top);
 
 	// iOS should always be fullscreen and at the origin
 	InWindow->SetCachedSize( FVector2D(Width, Height) );
@@ -88,8 +88,8 @@ void FSlateOpenGLViewport::Resize( int Width, int Height, bool bInFullscreen )
 			Swap<float>(Frame.size.width, Frame.size.height);
 		}
 
-		ViewportWidth = FMath::Trunc(Frame.size.width * Scale);
-		ViewportHeight = FMath::Trunc(Frame.size.height * Scale);
+		ViewportWidth = FMath::TruncToInt(Frame.size.width * Scale);
+		ViewportHeight = FMath::TruncToInt(Frame.size.height * Scale);
 	}
 
 	View.frame = CGRectMake(0, 0, ViewportWidth, ViewportHeight);
diff --git a/Engine/Source/Developer/TaskGraph/Private/STimeline.cpp b/Engine/Source/Developer/TaskGraph/Private/STimeline.cpp
--- a/Engine/Source/Developer/TaskGraph/Private/STimeline.cpp
+++ b/Engine/Source/Developer/TaskGraph/Private/STimeline.cpp
@@ -51,10 +51,10 @@ int32 STimeline::OnPaint( const FGeometry& AllottedGeometry, const FSlateRect& M
 	//);	
 	
 	// Create line points
-	const float RoundedMax = FMath::Ceil( MaxValue );
-	const float RoundedMin = FMath::Floor( MinValue );
+	const float RoundedMax = FMath::CeilToInt( MaxValue );
+	const float RoundedMin = FMath::FloorToInt( MinValue );
 	const float TimeScale = MaxValue - MinValue;
-	const int32 NumValues = FMath::Floor( AllottedGeometry.Size.X * Zoom / FixedLabelSpacing );
+	const int32 NumValues = FMath::FloorToInt( AllottedGeometry.Size.X * Zoom / FixedLabelSpacing );
 
 	TArray< FVector2D > LinePoints;
 	LinePoints.AddUninitialized( 2 );
diff --git a/Engine/Source/Developer/TextureCompressor/Private/TextureCompressorModule.cpp b/Engine/Source/Developer/TextureCompressor/Private/TextureCompressorModule.cpp
--- a/Engine/Source/Developer/TextureCompressor/Private/TextureCompressorModule.cpp
+++ b/Engine/Source/Developer/TextureCompressor/Private/TextureCompressorModule.cpp
@@ -380,7 +380,7 @@ static void GenerateSharpenedMipB8G8R8A8Templ(
 
 				if ( FilteredColor.A > AlphaThreshold )
 				{
-					FilteredColor.A = FMath::Trunc( FMath::Lerp( MinRandomAlpha, MaxRandomAlpha, RandomStream.GetFraction() ) );
+					FilteredColor.A = FMath::TruncToInt( FMath::Lerp( MinRandomAlpha, MaxRandomAlpha, RandomStream.GetFraction() ) );
 				}
 			}
 
@@ -1103,7 +1103,7 @@ static void GenerateAngularFilteredMips(TArray<FImage>& InOutMipChain, int32 Num
 		//	float FloatInputMip = FMath::Log2(FMath::Sqrt(AreaCoveredInNormalizedArea)) + InputMipCount - QualityBias;
 		// optimized
 		float FloatInputMip = 0.5f * FMath::Log2(AreaCoveredInNormalizedArea) + NumMips - QualityBias;
-		uint32 InputMip = FMath::Clamp(FMath::Trunc(FloatInputMip), 0, NumMips - 1);
+		uint32 InputMip = FMath::Clamp(FMath::TruncToInt(FloatInputMip), 0, NumMips - 1);
 
 		FImage* Mip = new(InOutMipChain) FImage(Extent, Extent, 6, ERawImageFormat::RGBA32F);
 		GenerateAngularFilteredMip(Mip, SrcMipChain[InputMip], ConeAngle);
diff --git a/Engine/Source/Developer/TextureFormatPVR/Private/TextureFormatPVR.cpp b/Engine/Source/Developer/TextureFormatPVR/Private/TextureFormatPVR.cpp
--- a/Engine/Source/Developer/TextureFormatPVR/Private/TextureFormatPVR.cpp
+++ b/Engine/Source/Developer/TextureFormatPVR/Private/TextureFormatPVR.cpp
@@ -87,14 +87,14 @@ static bool SquarifyImage(FImage& Image, uint32 MinSquareSize)
 	if(MultX == 1 || MultY == 1)
 	{
 		float FOverhead = float(FMath::Min(Image.SizeX, Image.SizeY)) / float(SquareSize);
-		int32 POverhead = FMath::Round(100.0f - (FOverhead * 100.0f));
+		int32 POverhead = FMath::RoundToInt(100.0f - (FOverhead * 100.0f));
 		UE_LOG(LogTextureFormatPVR, Warning, TEXT("Expanding mip (%d,%d) to (%d, %d). Memory overhead: ~%d%%"),
 			Image.SizeX, Image.SizeY, SquareSize, SquareSize, POverhead);
 	}
 	else if (MultX != MultY)
 	{
 		float FOverhead = float(FMath::Min(Image.SizeX, Image.SizeY)) / float(FMath::Max(Image.SizeX, Image.SizeY));
-		int32 POverhead = FMath::Round(100.0f - (FOverhead * 100.0f));
+		int32 POverhead = FMath::RoundToInt(100.0f - (FOverhead * 100.0f));
 		UE_LOG(LogTextureFormatPVR, Warning, TEXT("Expanding mip (%d,%d) to (%d, %d). Memory overhead: ~%d%%"),
 			Image.SizeX, Image.SizeY, FMath::Max(Image.SizeX, Image.SizeY), FMath::Max(Image.SizeX, Image.SizeY), POverhead);
 	}
@@ -154,7 +154,7 @@ static void DeriveNormalZ(FImage& Image)
 				const float NormalX = SourceColor.R / 255.0f * 2 - 1;
 				const float NormalY = SourceColor.G / 255.0f * 2 - 1;
 				const float NormalZ = FMath::Sqrt(FMath::Clamp<float>(1 - (NormalX * NormalX + NormalY * NormalY), 0, 1));
-				SourceColor.B = FMath::Trunc((NormalZ + 1) / 2.0f * 255.0f);
+				SourceColor.B = FMath::TruncToInt((NormalZ + 1) / 2.0f * 255.0f);
 			}
 		}
 	}
diff --git a/Engine/Source/Editor/BehaviorTreeEditor/Private/BehaviorTreeConnectionDrawingPolicy.cpp b/Engine/Source/Editor/BehaviorTreeEditor/Private/BehaviorTreeConnectionDrawingPolicy.cpp
--- a/Engine/Source/Editor/BehaviorTreeEditor/Private/BehaviorTreeConnectionDrawingPolicy.cpp
+++ b/Engine/Source/Editor/BehaviorTreeEditor/Private/BehaviorTreeConnectionDrawingPolicy.cpp
@@ -198,7 +198,7 @@ void FBehaviorTreeConnectionDrawingPolicy::DrawConnection(int32 LayerId, const F
 
 		float Time = (FPlatformTime::Seconds() - GStartTime);
 		const float BubbleOffset = FMath::Fmod(Time * BubbleSpeed, BubbleSpacing);
-		const int32 NumBubbles = FMath::Ceil(SplineLength/BubbleSpacing);
+		const int32 NumBubbles = FMath::CeilToInt(SplineLength/BubbleSpacing);
 		for (int32 i = 0; i < NumBubbles; ++i)
 		{
 			const float Distance = ((float)i * BubbleSpacing) + BubbleOffset;
diff --git a/Engine/Source/Editor/BehaviorTreeEditor/Private/SGraphNode_BehaviorTree.cpp b/Engine/Source/Editor/BehaviorTreeEditor/Private/SGraphNode_BehaviorTree.cpp
--- a/Engine/Source/Editor/BehaviorTreeEditor/Private/SGraphNode_BehaviorTree.cpp
+++ b/Engine/Source/Editor/BehaviorTreeEditor/Private/SGraphNode_BehaviorTree.cpp
@@ -477,7 +477,7 @@ void SGraphNode_BehaviorTree::Tick( const FGeometry& AllottedGeometry, const dou
 				{
 					NewFlashAlpha =
 						(DebuggerStateDuration > FlashStartTime + SearchPathBlink) ? 1.0f :
-						(FMath::Trunc(DebuggerStateDuration * SearchPathBlinkFreq) % 2) ? 1.0f : 0.0f;
+						(FMath::TruncToInt(DebuggerStateDuration * SearchPathBlinkFreq) % 2) ? 1.0f : 0.0f;
 				} 
 			}
 
diff --git a/Engine/Source/Editor/Cascade/Private/Cascade.cpp b/Engine/Source/Editor/Cascade/Private/Cascade.cpp
--- a/Engine/Source/Editor/Cascade/Private/Cascade.cpp
+++ b/Engine/Source/Editor/Cascade/Private/Cascade.cpp
@@ -4444,7 +4444,7 @@ void FCascade::OnSetRandomSeed()
 		ParticleSystem->PreEditChange(NULL);
 		ParticleSystemComponent->PreEditChange(NULL);
 
-		int32 RandomSeed = FMath::Round(RAND_MAX * FMath::SRand());
+		int32 RandomSeed = FMath::RoundToInt(RAND_MAX * FMath::SRand());
 		if (SelectedModule->SetRandomSeedEntry(0, RandomSeed) == false)
 		{
 			UE_LOG(LogCascade, Warning, TEXT("Failed to set random seed entry on module %s"), *(SelectedModule->GetClass()->GetName()));
@@ -4591,7 +4591,7 @@ bool FCascade::ConvertModuleToSeeded(UParticleSystem* ParticleSystem, UParticleE
 			if (RandSeedInfo != NULL)
 			{
 				RandSeedInfo->bResetSeedOnEmitterLooping = true;
-				RandSeedInfo->RandomSeeds.Add(FMath::Trunc(FMath::Rand() * UINT_MAX));
+				RandSeedInfo->RandomSeeds.Add(FMath::TruncToInt(FMath::Rand() * UINT_MAX));
 			}
 		}
 
diff --git a/Engine/Source/Editor/Cascade/Private/CascadeEmitterCanvasClient.cpp b/Engine/Source/Editor/Cascade/Private/CascadeEmitterCanvasClient.cpp
--- a/Engine/Source/Editor/Cascade/Private/CascadeEmitterCanvasClient.cpp
+++ b/Engine/Source/Editor/Cascade/Private/CascadeEmitterCanvasClient.cpp
@@ -208,7 +208,7 @@ void FCascadeEmitterCanvasClient::Draw(FViewport* Viewport, FCanvas* Canvas)
 	// Draw module errors and warnings.
 	{
 		UFont* ErrorFont = GEngine->GetSmallFont();
-		const int32 LineHeight = FMath::Trunc( ErrorFont->GetMaxCharHeight() );
+		const int32 LineHeight = FMath::TruncToInt( ErrorFont->GetMaxCharHeight() );
 		int32 DrawY = ViewY - 2 - LineHeight;
 
 		for (int32 i = 0; i < ModuleErrorStrings.Num(); ++i)
diff --git a/Engine/Source/Editor/DetailCustomizations/Private/EnvQueryParamInstanceCustomization.cpp b/Engine/Source/Editor/DetailCustomizations/Private/EnvQueryParamInstanceCustomization.cpp
--- a/Engine/Source/Editor/DetailCustomizations/Private/EnvQueryParamInstanceCustomization.cpp
+++ b/Engine/Source/Editor/DetailCustomizations/Private/EnvQueryParamInstanceCustomization.cpp
@@ -102,7 +102,7 @@ void FEnvQueryParamInstanceCustomization::OnParamNumValueChanged(float FloatValu
 	}
 	else if (ParamType == EEnvQueryParam::Int)
 	{
-		const int32 IntValue = FMath::Trunc(FloatValue);
+		const int32 IntValue = FMath::TruncToInt(FloatValue);
 		const float StoreValue = *((float*)&IntValue);
 		ValueProp->SetValue(StoreValue);
 		CachedInt = IntValue;
@@ -174,7 +174,7 @@ void FEnvQueryParamInstanceCustomization::InitCachedTypes()
 		switch (ParamType)
 		{
 		case EEnvQueryParam::Float:	CachedFloat = GetParamNumValue().GetValue(); break;
-		case EEnvQueryParam::Int:	CachedInt = FMath::Trunc(GetParamNumValue().GetValue()); break;
+		case EEnvQueryParam::Int:	CachedInt = FMath::TruncToInt(GetParamNumValue().GetValue()); break;
 		case EEnvQueryParam::Bool:	CachedBool = (GetParamBoolValue() == ESlateCheckBoxState::Checked); break;
 		}
 	}
diff --git a/Engine/Source/Editor/DistCurveEditor/Private/CurveEditorViewportClient.cpp b/Engine/Source/Editor/DistCurveEditor/Private/CurveEditorViewportClient.cpp
--- a/Engine/Source/Editor/DistCurveEditor/Private/CurveEditorViewportClient.cpp
+++ b/Engine/Source/Editor/DistCurveEditor/Private/CurveEditorViewportClient.cpp
@@ -982,13 +982,13 @@ void FCurveEditorViewportClient::DrawEntry(FViewport* Viewport, FCanvas* Canvas,
 				// otherwise, we end up w/ 100,000s of steps.
 				float Scalar = 1.0f;
 				float KeyDiffTemp = KeyDiff;
-				while (FMath::Trunc(KeyDiffTemp / Scalar) > 1.0f)
+				while (FMath::TruncToInt(KeyDiffTemp / Scalar) > 1)
 				{
 					Scalar *= 10.0f;
 				}
 				float DrawTrackInRes = CurveDrawRes/PixelsPerIn;
 				DrawTrackInRes *= Scalar;
-				int32 NumSteps = FMath::Ceil(KeyDiff/DrawTrackInRes);
+				int32 NumSteps = FMath::CeilToInt(KeyDiff/DrawTrackInRes);
 
 				if(Scalar > 1.0f)
 				{
@@ -1126,8 +1126,8 @@ void FCurveEditorViewportClient::DrawEntry(FViewport* Viewport, FCanvas* Canvas,
 					FVector2D HandleDir = CalcTangentDir((PixelsPerOut/PixelsPerIn) * ArriveTangent);
 
 					FIntPoint HandlePos;
-					HandlePos.X = NewKeyPos.X - FMath::Round(HandleDir.X * HandleLength);
-					HandlePos.Y = NewKeyPos.Y - FMath::Round(HandleDir.Y * HandleLength);
+					HandlePos.X = NewKeyPos.X - FMath::RoundToInt(HandleDir.X * HandleLength);
+					HandlePos.Y = NewKeyPos.Y - FMath::RoundToInt(HandleDir.Y * HandleLength);
 					LineItem.SetColor( FLinearColor::White );
 					LineItem.Draw( Canvas, NewKeyPos, HandlePos );
 					
@@ -1143,8 +1143,8 @@ void FCurveEditorViewportClient::DrawEntry(FViewport* Viewport, FCanvas* Canvas,
 					FVector2D HandleDir = CalcTangentDir((PixelsPerOut/PixelsPerIn) * LeaveTangent);
 
 					FIntPoint HandlePos;
-					HandlePos.X = NewKeyPos.X + FMath::Round(HandleDir.X * HandleLength);
-					HandlePos.Y = NewKeyPos.Y + FMath::Round(HandleDir.Y * HandleLength);
+					HandlePos.X = NewKeyPos.X + FMath::RoundToInt(HandleDir.X * HandleLength);
+					HandlePos.Y = NewKeyPos.Y + FMath::RoundToInt(HandleDir.Y * HandleLength);
 
 					LineItem.SetColor( FLinearColor::White );
 					LineItem.Draw( Canvas, NewKeyPos, HandlePos );
@@ -1163,7 +1163,7 @@ void FCurveEditorViewportClient::DrawEntry(FViewport* Viewport, FCanvas* Canvas,
 				FString KeyComment;
 				if (bSnapToFrames)
 				{
-					KeyComment = FString::Printf(TEXT("(%df,%3.2f)"), FMath::Round(NewKey.X/InSnapAmount), NewKey.Y);
+					KeyComment = FString::Printf(TEXT("(%df,%3.2f)"), FMath::RoundToInt(NewKey.X/InSnapAmount), NewKey.Y);
 				}
 				else
 				{
@@ -1194,7 +1194,7 @@ namespace CurveEditor
 	/** Calculate the best frames' density. */
 	uint32 CalculateBestFrameStep(float SnapAmount, float PixelsPerSec, float MinPixelsPerGrid)
 	{
-		uint32 FrameRate = FMath::Ceil(1.0f / SnapAmount);
+		uint32 FrameRate = FMath::CeilToInt(1.0f / SnapAmount);
 		uint32 FrameStep = 1;
 
 		// Calculate minimal-symmetric integer divisor.
@@ -1278,7 +1278,7 @@ void FCurveEditorViewportClient::DrawGrid(FViewport* Viewport, FCanvas* Canvas)
 
 	// Draw input grid
 	FCanvasLineItem LineItem;
-	int32 InNum = FMath::Floor(SharedData->StartIn/InGridSpacing);
+	int32 InNum = FMath::FloorToInt(SharedData->StartIn/InGridSpacing);
 	while(InNum*InGridSpacing < SharedData->EndIn)
 	{
 		FColor LineColor = GridColor;
@@ -1305,7 +1305,7 @@ void FCurveEditorViewportClient::DrawGrid(FViewport* Viewport, FCanvas* Canvas)
 	}
 
 	// Draw output grid
-	int32 OutNum = FMath::Floor(SharedData->StartOut/OutGridSpacing);
+	int32 OutNum = FMath::FloorToInt(SharedData->StartOut/OutGridSpacing);
 	while(OutNum*OutGridSpacing < SharedData->EndOut)
 	{
 		FIntPoint GridPos = CalcScreenPos(FVector2D(0.f, OutNum*OutGridSpacing));
@@ -1324,7 +1324,7 @@ void FCurveEditorViewportClient::DrawGrid(FViewport* Viewport, FCanvas* Canvas)
 	// Draw input labels
 	FCanvasTextItem TextItem( FVector2D::ZeroVector, FText::GetEmpty(), GEditor->GetSmallFont(), GridTextColor );
 	
-	InNum = FMath::Floor(SharedData->StartIn/InGridSpacing);
+	InNum = FMath::FloorToInt(SharedData->StartIn/InGridSpacing);
 	while(InNum*InGridSpacing < SharedData->EndIn)
 	{
 		// Draw value label
@@ -1352,7 +1352,7 @@ void FCurveEditorViewportClient::DrawGrid(FViewport* Viewport, FCanvas* Canvas)
 
 	// Draw output labels
 
-	OutNum = FMath::Floor(SharedData->StartOut/OutGridSpacing);
+	OutNum = FMath::FloorToInt(SharedData->StartOut/OutGridSpacing);
 	while(OutNum*OutGridSpacing < SharedData->EndOut)
 	{
 		FIntPoint GridPos = CalcScreenPos(FVector2D(0.f, OutNum*OutGridSpacing));
@@ -1395,19 +1395,19 @@ FColor FCurveEditorViewportClient::GetLineColor(FCurveEdInterface* EdInterface,
 
 			Value = EdInterface->EvalSub(0, InVal);
 			Value *= 255.9f;
-			StepColor.R = FMath::Trunc(FMath::Clamp<float>(Value, 0.f, 255.9f));
+			StepColor.R = FMath::TruncToInt(FMath::Clamp<float>(Value, 0.f, 255.9f));
 			Value = EdInterface->EvalSub(1, InVal);
 			Value *= 255.9f;
-			StepColor.G = FMath::Trunc(FMath::Clamp<float>(Value, 0.f, 255.9f));
+			StepColor.G = FMath::TruncToInt(FMath::Clamp<float>(Value, 0.f, 255.9f));
 			Value = EdInterface->EvalSub(2, InVal);
 			Value *= 255.9f;
-			StepColor.B = FMath::Trunc(FMath::Clamp<float>(Value, 0.f, 255.9f));
+			StepColor.B = FMath::TruncToInt(FMath::Clamp<float>(Value, 0.f, 255.9f));
 		}
 		else
 		{
-			StepColor.R = FMath::Trunc(FMath::Clamp<float>(EdInterface->EvalSub(0, InVal), 0.f, 255.9f));
-			StepColor.G = FMath::Trunc(FMath::Clamp<float>(EdInterface->EvalSub(1, InVal), 0.f, 255.9f));
-			StepColor.B = FMath::Trunc(FMath::Clamp<float>(EdInterface->EvalSub(2, InVal), 0.f, 255.9f));
+			StepColor.R = FMath::TruncToInt(FMath::Clamp<float>(EdInterface->EvalSub(0, InVal), 0.f, 255.9f));
+			StepColor.G = FMath::TruncToInt(FMath::Clamp<float>(EdInterface->EvalSub(1, InVal), 0.f, 255.9f));
+			StepColor.B = FMath::TruncToInt(FMath::Clamp<float>(EdInterface->EvalSub(2, InVal), 0.f, 255.9f));
 		}
 		StepColor.A = 255;
 	}
@@ -1419,11 +1419,11 @@ FColor FCurveEditorViewportClient::GetLineColor(FCurveEdInterface* EdInterface,
 
 			Value = EdInterface->EvalSub(0, InVal);
 			Value *= 255.9f;
-			StepColor.R = FMath::Trunc(FMath::Clamp<float>(Value, 0.f, 255.9f));
+			StepColor.R = FMath::TruncToInt(FMath::Clamp<float>(Value, 0.f, 255.9f));
 		}
 		else
 		{
-			StepColor.R = FMath::Trunc(FMath::Clamp<float>(EdInterface->EvalSub(0, InVal), 0.f, 255.9f));
+			StepColor.R = FMath::TruncToInt(FMath::Clamp<float>(EdInterface->EvalSub(0, InVal), 0.f, 255.9f));
 		}
 		StepColor.G = StepColor.R;
 		StepColor.B = StepColor.R;
@@ -1688,7 +1688,7 @@ float FCurveEditorViewportClient::SnapIn(float InValue)
 {
 	if(bSnapEnabled)
 	{
-		return InSnapAmount * FMath::Round(InValue/InSnapAmount);
+		return InSnapAmount * FMath::RoundToFloat(InValue/InSnapAmount);
 	}
 	else
 	{
diff --git a/Engine/Source/Editor/DistCurveEditor/Private/SDistributionCurveEditor.cpp b/Engine/Source/Editor/DistCurveEditor/Private/SDistributionCurveEditor.cpp
--- a/Engine/Source/Editor/DistCurveEditor/Private/SDistributionCurveEditor.cpp
+++ b/Engine/Source/Editor/DistCurveEditor/Private/SDistributionCurveEditor.cpp
@@ -597,17 +597,17 @@ void SDistributionCurveEditor::OnSetColor()
 		float Value;
 
 		Value	= EdInterface->GetKeyOut(0, SelKey.KeyIndex) * 255.9f;
-		InputColor.R = FMath::Trunc(Value);
+		InputColor.R = FMath::TruncToInt(Value);
 		Value	= EdInterface->GetKeyOut(1, SelKey.KeyIndex) * 255.9f;
-		InputColor.G = FMath::Trunc(Value);
+		InputColor.G = FMath::TruncToInt(Value);
 		Value	= EdInterface->GetKeyOut(2, SelKey.KeyIndex) * 255.9f;
-		InputColor.B = FMath::Trunc(Value);
+		InputColor.B = FMath::TruncToInt(Value);
 	}
 	else
 	{
-		InputColor.R = FMath::Trunc(FMath::Clamp<float>(EdInterface->GetKeyOut(0, SelKey.KeyIndex), 0.f, 255.9f));
-		InputColor.G = FMath::Trunc(FMath::Clamp<float>(EdInterface->GetKeyOut(1, SelKey.KeyIndex), 0.f, 255.9f));
-		InputColor.B = FMath::Trunc(FMath::Clamp<float>(EdInterface->GetKeyOut(2, SelKey.KeyIndex), 0.f, 255.9f));
+		InputColor.R = FMath::TruncToInt(FMath::Clamp<float>(EdInterface->GetKeyOut(0, SelKey.KeyIndex), 0.f, 255.9f));
+		InputColor.G = FMath::TruncToInt(FMath::Clamp<float>(EdInterface->GetKeyOut(1, SelKey.KeyIndex), 0.f, 255.9f));
+		InputColor.B = FMath::TruncToInt(FMath::Clamp<float>(EdInterface->GetKeyOut(2, SelKey.KeyIndex), 0.f, 255.9f));
 	}
 
 	//since the data isn't stored in standard colors, a temp color is used
diff --git a/Engine/Source/Editor/EnvironmentQueryEditor/Private/DetailCustomizations/EnvQueryParamSetupCustomization.cpp b/Engine/Source/Editor/EnvironmentQueryEditor/Private/DetailCustomizations/EnvQueryParamSetupCustomization.cpp
--- a/Engine/Source/Editor/EnvironmentQueryEditor/Private/DetailCustomizations/EnvQueryParamSetupCustomization.cpp
+++ b/Engine/Source/Editor/EnvironmentQueryEditor/Private/DetailCustomizations/EnvQueryParamSetupCustomization.cpp
@@ -140,7 +140,7 @@ void FEnvQueryParamSetupCustomization::OnParamNumValueChanged(float FloatValue)
 	}
 	else if (ParamType == EEnvQueryParam::Int)
 	{
-		const int32 IntValue = FMath::Trunc(FloatValue);
+		const int32 IntValue = FMath::TruncToInt(FloatValue);
 		ValueProp->SetValue(IntValue);
 	}
 }
diff --git a/Engine/Source/Editor/FoliageEdit/Private/FoliageEdMode.cpp b/Engine/Source/Editor/FoliageEdit/Private/FoliageEdMode.cpp
--- a/Engine/Source/Editor/FoliageEdit/Private/FoliageEdMode.cpp
+++ b/Engine/Source/Editor/FoliageEdit/Private/FoliageEdMode.cpp
@@ -547,14 +547,14 @@ static bool CheckVertexColor(const UInstancedFoliageSettings* MeshSettings, cons
 
 	if( MeshSettings->VertexColorMaskInvert )
 	{
-		if( ColorChannel > FMath::Round(MeshSettings->VertexColorMaskThreshold * 255.f) )
+		if( ColorChannel > FMath::RoundToInt(MeshSettings->VertexColorMaskThreshold * 255.f) )
 		{
 			return false;
 		}
 	}
 	else
 	{
-		if( ColorChannel < FMath::Round(MeshSettings->VertexColorMaskThreshold * 255.f) )
+		if( ColorChannel < FMath::RoundToInt(MeshSettings->VertexColorMaskThreshold * 255.f) )
 		{
 			return false;
 		}
@@ -597,7 +597,7 @@ void FEdModeFoliage::AddInstancesForBrush( UWorld* InWorld, AInstancedFoliageAct
 						float HitWeight = HitLandscape->GetLayerWeightAtLocation( Instance.Location, HitLandscape->GetLandscapeInfo()->GetLayerInfoByName(LandscapeLayerName), LayerCache );
 
 						// Add count to bucket.
-						ExistingInstanceBuckets[FMath::Round(HitWeight * (float)(NUM_INSTANCE_BUCKETS-1))]++;
+						ExistingInstanceBuckets[FMath::RoundToInt(HitWeight * (float)(NUM_INSTANCE_BUCKETS-1))]++;
 					}
 				}
 			}
@@ -692,7 +692,7 @@ void FEdModeFoliage::AddInstancesForBrush( UWorld* InWorld, AInstancedFoliageAct
 					}
 				}
 
-				new(PotentialInstanceBuckets[FMath::Round(HitWeight * (float)(NUM_INSTANCE_BUCKETS-1))]) FPotentialInstance( Hit.Location, Hit.Normal, Hit.Component.Get(), HitWeight );
+				new(PotentialInstanceBuckets[FMath::RoundToInt(HitWeight * (float)(NUM_INSTANCE_BUCKETS-1))]) FPotentialInstance( Hit.Location, Hit.Normal, Hit.Component.Get(), HitWeight );
 			}
 		}
 
@@ -703,7 +703,7 @@ void FEdModeFoliage::AddInstancesForBrush( UWorld* InWorld, AInstancedFoliageAct
 
 			// We use the number that actually succeeded in placement (due to parameters) as the target
 			// for the number that should be in the brush region.
-			int32 AdditionalInstances = FMath::Clamp<int32>( FMath::Round( BucketFraction * (float)(PotentialInstances.Num() - ExistingInstanceBuckets[BucketIdx]) * Pressure), 0, PotentialInstances.Num() );
+			int32 AdditionalInstances = FMath::Clamp<int32>( FMath::RoundToInt( BucketFraction * (float)(PotentialInstances.Num() - ExistingInstanceBuckets[BucketIdx]) * Pressure), 0, PotentialInstances.Num() );
 			for( int32 Idx=0;Idx<AdditionalInstances;Idx++ )
 			{
 				FFoliageInstance Inst;
@@ -719,7 +719,7 @@ void FEdModeFoliage::AddInstancesForBrush( UWorld* InWorld, AInstancedFoliageAct
 /** Remove instances inside the brush to match DesiredInstanceCount */
 void FEdModeFoliage::RemoveInstancesForBrush( AInstancedFoliageActor* IFA, FFoliageMeshInfo& MeshInfo, int32 DesiredInstanceCount, TArray<int32>& ExistingInstances, float Pressure )
 {
-	int32 InstancesToRemove = FMath::Round( (float)(ExistingInstances.Num() - DesiredInstanceCount) * Pressure);
+	int32 InstancesToRemove = FMath::RoundToInt( (float)(ExistingInstances.Num() - DesiredInstanceCount) * Pressure);
 	int32 InstancesToKeep = ExistingInstances.Num() - InstancesToRemove;
 	if( InstancesToKeep > 0 )
 	{						
@@ -1131,7 +1131,7 @@ void FEdModeFoliage::ApplyBrush( FLevelEditorViewportClient* ViewportClient )
 					if( SnapShot ) 
 					{
 						// Use snapshot to determine number of instances at the start of the brush stroke
-						int32 NewInstanceCount = FMath::Round( (float)SnapShot->CountInstancesInsideSphere(BrushSphere) * MeshSettings->ReapplyDensityAmount );
+						int32 NewInstanceCount = FMath::RoundToInt( (float)SnapShot->CountInstancesInsideSphere(BrushSphere) * MeshSettings->ReapplyDensityAmount );
 						if( MeshSettings->ReapplyDensityAmount > 1.f && NewInstanceCount > Instances.Num() )
 						{
 							AddInstancesForBrush( World , IFA, MeshIt.Key(), MeshInfo, NewInstanceCount, Instances, Pressure );
@@ -1153,7 +1153,7 @@ void FEdModeFoliage::ApplyBrush( FLevelEditorViewportClient* ViewportClient )
 				// Shift unpaints
 				if( IsShiftDown(ViewportClient->Viewport) )
 				{
-					int32 DesiredInstanceCount =  FMath::Round(BrushArea * MeshSettings->Density * UISettings.GetUnpaintDensity() / (1000.f*1000.f));
+					int32 DesiredInstanceCount =  FMath::RoundToInt(BrushArea * MeshSettings->Density * UISettings.GetUnpaintDensity() / (1000.f*1000.f));
 
 					if( DesiredInstanceCount < Instances.Num() )
 					{
@@ -1166,7 +1166,7 @@ void FEdModeFoliage::ApplyBrush( FLevelEditorViewportClient* ViewportClient )
 					float DesiredInstanceCountFloat = BrushArea * MeshSettings->Density * UISettings.GetPaintDensity() / (1000.f*1000.f);
 
 					// Allow a single instance with a random chance, if the brush is smaller than the density
-					int32 DesiredInstanceCount = DesiredInstanceCountFloat > 1.f ? FMath::Round(DesiredInstanceCountFloat) : FMath::FRand() < DesiredInstanceCountFloat ? 1 : 0;
+					int32 DesiredInstanceCount = DesiredInstanceCountFloat > 1.f ? FMath::RoundToInt(DesiredInstanceCountFloat) : FMath::FRand() < DesiredInstanceCountFloat ? 1 : 0;
 
 					AddInstancesForBrush( World, IFA, MeshIt.Key(), MeshInfo, DesiredInstanceCount, Instances, Pressure );
 				}
@@ -1354,7 +1354,7 @@ void FEdModeFoliage::ApplyPaintBucket(AActor* Actor, bool bRemove)
 						float DesiredInstanceCountFloat = Triangle.Area * MeshSettings->Density * UISettings.GetPaintDensity() / (1000.f*1000.f);
 
 						// Allow a single instance with a random chance, if the brush is smaller than the density
-						int32 DesiredInstanceCount = DesiredInstanceCountFloat > 1.f ? FMath::Round(DesiredInstanceCountFloat) : FMath::FRand() < DesiredInstanceCountFloat ? 1 : 0;
+						int32 DesiredInstanceCount = DesiredInstanceCountFloat > 1.f ? FMath::RoundToInt(DesiredInstanceCountFloat) : FMath::FRand() < DesiredInstanceCountFloat ? 1 : 0;
 				
 						for( int32 Idx=0;Idx<DesiredInstanceCount;Idx++ )
 						{
diff --git a/Engine/Source/Editor/GraphEditor/Private/AnimationNodes/SGraphNodeSequencePlayer.cpp b/Engine/Source/Editor/GraphEditor/Private/AnimationNodes/SGraphNodeSequencePlayer.cpp
--- a/Engine/Source/Editor/GraphEditor/Private/AnimationNodes/SGraphNodeSequencePlayer.cpp
+++ b/Engine/Source/Editor/GraphEditor/Private/AnimationNodes/SGraphNodeSequencePlayer.cpp
@@ -31,9 +31,9 @@ FString SGraphNodeSequencePlayer::GetPositionTooltip() const
 	int32 FrameCount;
 	if (GetSequencePositionInfo(/*out*/ Position, /*out*/ Length, /*out*/ FrameCount))
 	{
-		const int32 Minutes = FMath::Trunc(Position/60.0f);
-		const int32 Seconds = FMath::Trunc(Position) % 60;
-		const int32 Hundredths = FMath::Trunc(FMath::Fractional(Position)*100);
+		const int32 Minutes = FMath::TruncToInt(Position/60.0f);
+		const int32 Seconds = FMath::TruncToInt(Position) % 60;
+		const int32 Hundredths = FMath::TruncToInt(FMath::Fractional(Position)*100);
 
 		FString MinuteStr;
 		if (Minutes > 0)
@@ -45,7 +45,7 @@ FString SGraphNodeSequencePlayer::GetPositionTooltip() const
 
 		const FString HundredthsStr = FString::Printf(TEXT(".%02d"), Hundredths);
 
-		const int32 CurrentFrame = FMath::Trunc((Position / Length) * FrameCount);
+		const int32 CurrentFrame = FMath::TruncToInt((Position / Length) * FrameCount);
 
 		const FString FramesStr = FString::Printf(TEXT("Frame %d"), CurrentFrame);
 
diff --git a/Engine/Source/Editor/GraphEditor/Private/ConnectionDrawingPolicy.cpp b/Engine/Source/Editor/GraphEditor/Private/ConnectionDrawingPolicy.cpp
--- a/Engine/Source/Editor/GraphEditor/Private/ConnectionDrawingPolicy.cpp
+++ b/Engine/Source/Editor/GraphEditor/Private/ConnectionDrawingPolicy.cpp
@@ -354,7 +354,7 @@ void FConnectionDrawingPolicy::DrawConnection( int32 LayerId, const FVector2D& S
 
 			float Time = (FPlatformTime::Seconds() - GStartTime);
 			const float BubbleOffset = FMath::Fmod(Time * BubbleSpeed, BubbleSpacing);
-			const int32 NumBubbles = FMath::Ceil(SplineLength/BubbleSpacing);
+			const int32 NumBubbles = FMath::CeilToInt(SplineLength/BubbleSpacing);
 			for (int32 i = 0; i < NumBubbles; ++i)
 			{
 				const float Distance = ((float)i * BubbleSpacing) + BubbleOffset;
diff --git a/Engine/Source/Editor/GraphEditor/Private/MaterialNodes/SGraphNodeMaterialBase.cpp b/Engine/Source/Editor/GraphEditor/Private/MaterialNodes/SGraphNodeMaterialBase.cpp
--- a/Engine/Source/Editor/GraphEditor/Private/MaterialNodes/SGraphNodeMaterialBase.cpp
+++ b/Engine/Source/Editor/GraphEditor/Private/MaterialNodes/SGraphNodeMaterialBase.cpp
@@ -87,16 +87,16 @@ void FPreviewViewport::OnDrawViewport( const FGeometry& AllottedGeometry, const
 	FSlateRect ClippedCanvasRect = SlateCanvasRect.IntersectionWith(MyClippingRect);
 
 	FIntRect CanvasRect(
-		FMath::Trunc( FMath::Max(0.0f, SlateCanvasRect.Left) ),
-		FMath::Trunc( FMath::Max(0.0f, SlateCanvasRect.Top) ),
-		FMath::Trunc( FMath::Max(0.0f, SlateCanvasRect.Right) ), 
-		FMath::Trunc( FMath::Max(0.0f, SlateCanvasRect.Bottom) ) );
+		FMath::TruncToInt( FMath::Max(0.0f, SlateCanvasRect.Left) ),
+		FMath::TruncToInt( FMath::Max(0.0f, SlateCanvasRect.Top) ),
+		FMath::TruncToInt( FMath::Max(0.0f, SlateCanvasRect.Right) ), 
+		FMath::TruncToInt( FMath::Max(0.0f, SlateCanvasRect.Bottom) ) );
 
 	FIntRect ClippingRect(
-		FMath::Trunc( FMath::Max(0.0f, ClippedCanvasRect.Left) ),
-		FMath::Trunc( FMath::Max(0.0f, ClippedCanvasRect.Top) ),
-		FMath::Trunc( FMath::Max(0.0f, ClippedCanvasRect.Right) ), 
-		FMath::Trunc( FMath::Max(0.0f, ClippedCanvasRect.Bottom) ) );
+		FMath::TruncToInt( FMath::Max(0.0f, ClippedCanvasRect.Left) ),
+		FMath::TruncToInt( FMath::Max(0.0f, ClippedCanvasRect.Top) ),
+		FMath::TruncToInt( FMath::Max(0.0f, ClippedCanvasRect.Right) ), 
+		FMath::TruncToInt( FMath::Max(0.0f, ClippedCanvasRect.Bottom) ) );
 
 	bool bIsRealtime = MaterialNode->RealtimeDelegate.IsBound() ? MaterialNode->RealtimeDelegate.Execute() : false;
 
diff --git a/Engine/Source/Editor/GraphEditor/Private/SGraphNodeComment.cpp b/Engine/Source/Editor/GraphEditor/Private/SGraphNodeComment.cpp
--- a/Engine/Source/Editor/GraphEditor/Private/SGraphNodeComment.cpp
+++ b/Engine/Source/Editor/GraphEditor/Private/SGraphNodeComment.cpp
@@ -131,8 +131,8 @@ FReply SGraphNodeComment::OnMouseMove(const FGeometry& MyGeometry, const FPointe
 
 		const FVector2D PreviousUserSize = UserSize;
 		const float SnapSize = SNodePanel::GetSnapGridSize();
-		UserSize.X = SnapSize * FMath::Round(StoredUserSize.X/SnapSize);
-		UserSize.Y = SnapSize * FMath::Round(StoredUserSize.Y/SnapSize);
+		UserSize.X = SnapSize * FMath::RoundToFloat(StoredUserSize.X/SnapSize);
+		UserSize.Y = SnapSize * FMath::RoundToFloat(StoredUserSize.Y/SnapSize);
 
 		FVector2D DeltaNodePos(0,0);
 		//Modify node position (resizing top and left sides)
@@ -468,8 +468,8 @@ FReply SGraphNodeComment::OnMouseButtonUp( const FGeometry& MyGeometry, const FP
 		bUserIsDragging = false;
 
 		// Resize the node	
-		UserSize.X = FPlatformMath::Round(UserSize.X);
-		UserSize.Y =FPlatformMath::Round(UserSize.Y);
+		UserSize.X = FPlatformMath::RoundToFloat(UserSize.X);
+		UserSize.Y =FPlatformMath::RoundToFloat(UserSize.Y);
 
 		GetNodeObj()->ResizeNode(UserSize);
 		return FReply::Handled().ReleaseMouseCapture();
diff --git a/Engine/Source/Editor/GraphEditor/Private/SNodePanel.cpp b/Engine/Source/Editor/GraphEditor/Private/SNodePanel.cpp
--- a/Engine/Source/Editor/GraphEditor/Private/SNodePanel.cpp
+++ b/Engine/Source/Editor/GraphEditor/Private/SNodePanel.cpp
@@ -606,7 +606,7 @@ FReply SNodePanel::OnMouseMove(const FGeometry& MyGeometry, const FPointerEvent&
 
 			TotalMouseDeltaY += CursorDelta.Y;
 
-			const int32 ZoomLevelDelta = FMath::Floor(TotalMouseDeltaY * NodePanelDefs::MouseZoomScaling);
+			const int32 ZoomLevelDelta = FMath::FloorToInt(TotalMouseDeltaY * NodePanelDefs::MouseZoomScaling);
 
 			// Get rid of mouse movement that's been 'used up' by zooming
 			if (ZoomLevelDelta != 0)
@@ -657,8 +657,8 @@ FReply SNodePanel::OnMouseMove(const FGeometry& MyGeometry, const FPointerEvent&
 
 						// Snap to grid
 						const float SnapSize = GetSnapGridSize();
-						AnchorNodeNewPos.X = SnapSize * FMath::Round(AnchorNodeNewPos.X/SnapSize);
-						AnchorNodeNewPos.Y = SnapSize * FMath::Round(AnchorNodeNewPos.Y/SnapSize);
+						AnchorNodeNewPos.X = SnapSize * FMath::RoundToFloat(AnchorNodeNewPos.X/SnapSize);
+						AnchorNodeNewPos.Y = SnapSize * FMath::RoundToFloat(AnchorNodeNewPos.Y/SnapSize);
 
 						// Dragging an unselected node automatically selects it.
 						SelectionManager.StartDraggingNode(NodeBeingDragged->GetObjectBeingDisplayed(), MouseEvent);
@@ -847,8 +847,8 @@ FReply SNodePanel::OnMouseButtonUp( const FGeometry& MyGeometry, const FPointerE
 			const FVector2D ScreenSpaceCursorPos = MyGeometry.LocalToAbsolute( GraphCoordToPanelCoord( SoftwareCursorPosition ) );
 
 			FIntPoint BestPositionInViewport(
-				FMath::Round( FMath::Clamp( ScreenSpaceCursorPos.X, ThisPanelScreenSpaceRect.Left, ThisPanelScreenSpaceRect.Right ) ),
-				FMath::Round( FMath::Clamp( ScreenSpaceCursorPos.Y, ThisPanelScreenSpaceRect.Top, ThisPanelScreenSpaceRect.Bottom ) )
+				FMath::RoundToInt( FMath::Clamp( ScreenSpaceCursorPos.X, ThisPanelScreenSpaceRect.Left, ThisPanelScreenSpaceRect.Right ) ),
+				FMath::RoundToInt( FMath::Clamp( ScreenSpaceCursorPos.Y, ThisPanelScreenSpaceRect.Top, ThisPanelScreenSpaceRect.Bottom ) )
 				);
 
 			if (!bCursorInDeadZone)
@@ -867,7 +867,7 @@ FReply SNodePanel::OnMouseWheel(const FGeometry& MyGeometry, const FPointerEvent
 {
 	// We want to zoom into this point; i.e. keep it the same fraction offset into the panel
 	const FVector2D WidgetSpaceCursorPos = MyGeometry.AbsoluteToLocal( MouseEvent.GetScreenSpacePosition() );
-	const int32 ZoomLevelDelta = FMath::Floor( MouseEvent.GetWheelDelta() );
+	const int32 ZoomLevelDelta = FMath::FloorToInt( MouseEvent.GetWheelDelta() );
 	ChangeZoomLevel(ZoomLevelDelta, WidgetSpaceCursorPos, MouseEvent.IsControlDown());
 
 	return FReply::Handled();
diff --git a/Engine/Source/Editor/LandscapeEditor/Classes/LandscapeEditorObject.h b/Engine/Source/Editor/LandscapeEditor/Classes/LandscapeEditorObject.h
--- a/Engine/Source/Editor/LandscapeEditor/Classes/LandscapeEditorObject.h
+++ b/Engine/Source/Editor/LandscapeEditor/Classes/LandscapeEditorObject.h
@@ -513,8 +513,8 @@ class ULandscapeEditorObject : public UObject
 	void NewLandscape_ClampSize()
 	{
 		// Max size is either whole components below 8192 verts, or 32 components
-		NewLandscape_ComponentCount.X = FMath::Clamp(NewLandscape_ComponentCount.X, 1, FMath::Min(32, FMath::Floor(8191 / (NewLandscape_SectionsPerComponent * NewLandscape_QuadsPerSection))));
-		NewLandscape_ComponentCount.Y = FMath::Clamp(NewLandscape_ComponentCount.Y, 1, FMath::Min(32, FMath::Floor(8191 / (NewLandscape_SectionsPerComponent * NewLandscape_QuadsPerSection))));
+		NewLandscape_ComponentCount.X = FMath::Clamp(NewLandscape_ComponentCount.X, 1, FMath::Min(32, FMath::FloorToInt(8191 / (NewLandscape_SectionsPerComponent * NewLandscape_QuadsPerSection))));
+		NewLandscape_ComponentCount.Y = FMath::Clamp(NewLandscape_ComponentCount.Y, 1, FMath::Min(32, FMath::FloorToInt(8191 / (NewLandscape_SectionsPerComponent * NewLandscape_QuadsPerSection))));
 	}
 
 	void UpdateComponentCount()
diff --git a/Engine/Source/Editor/LandscapeEditor/Private/LandscapeEdMode.cpp b/Engine/Source/Editor/LandscapeEditor/Private/LandscapeEdMode.cpp
--- a/Engine/Source/Editor/LandscapeEditor/Private/LandscapeEdMode.cpp
+++ b/Engine/Source/Editor/LandscapeEditor/Private/LandscapeEdMode.cpp
@@ -392,16 +392,16 @@ void FEdModeLandscape::Enter()
 			// Restore Sampled Data if exist...
 			if (CurrentGizmoActor->CachedScaleXY > 0.f)
 			{
-				int32 SizeX = FMath::Ceil(CurrentGizmoActor->CachedWidth  / CurrentGizmoActor->CachedScaleXY);
-				int32 SizeY = FMath::Ceil(CurrentGizmoActor->CachedHeight / CurrentGizmoActor->CachedScaleXY);
+				int32 SizeX = FMath::CeilToInt(CurrentGizmoActor->CachedWidth  / CurrentGizmoActor->CachedScaleXY);
+				int32 SizeY = FMath::CeilToInt(CurrentGizmoActor->CachedHeight / CurrentGizmoActor->CachedScaleXY);
 				for (int32 Y = 0; Y < CurrentGizmoActor->SampleSizeY; ++Y)
 				{
 					for (int32 X = 0; X < CurrentGizmoActor->SampleSizeX; ++X)
 					{
 						float TexX = X * SizeX / CurrentGizmoActor->SampleSizeX;
 						float TexY = Y * SizeY / CurrentGizmoActor->SampleSizeY;
-						int32 LX = FMath::Floor(TexX);
-						int32 LY = FMath::Floor(TexY);
+						int32 LX = FMath::FloorToInt(TexX);
+						int32 LY = FMath::FloorToInt(TexY);
 
 						float FracX = TexX - LX;
 						float FracY = TexY - LY;
@@ -1383,7 +1383,7 @@ bool FEdModeLandscape::InputDelta( FLevelEditorViewportClient* InViewportClient,
 			case ELandscapeEdge::X_Negative_Y_Positive:
 				{
 					const int32 InitialComponentCountX = UISettings->NewLandscape_ComponentCount.X;
-					const int32 Delta = FMath::Round(HitLocation.X + (float)InitialComponentCountX / 2);
+					const int32 Delta = FMath::RoundToInt(HitLocation.X + (float)InitialComponentCountX / 2);
 					UISettings->NewLandscape_ComponentCount.X = InitialComponentCountX - Delta;
 					UISettings->NewLandscape_ClampSize();
 					const int32 ActualDelta = UISettings->NewLandscape_ComponentCount.X - InitialComponentCountX;
@@ -1395,7 +1395,7 @@ bool FEdModeLandscape::InputDelta( FLevelEditorViewportClient* InViewportClient,
 			case ELandscapeEdge::X_Positive_Y_Positive:
 				{
 					const int32 InitialComponentCountX = UISettings->NewLandscape_ComponentCount.X;
-					int32 Delta = FMath::Round(HitLocation.X - (float)InitialComponentCountX / 2);
+					int32 Delta = FMath::RoundToInt(HitLocation.X - (float)InitialComponentCountX / 2);
 					UISettings->NewLandscape_ComponentCount.X = InitialComponentCountX + Delta;
 					UISettings->NewLandscape_ClampSize();
 					const int32 ActualDelta = UISettings->NewLandscape_ComponentCount.X - InitialComponentCountX;
@@ -1411,7 +1411,7 @@ bool FEdModeLandscape::InputDelta( FLevelEditorViewportClient* InViewportClient,
 			case ELandscapeEdge::X_Positive_Y_Negative:
 				{
 					const int32 InitialComponentCountY = UISettings->NewLandscape_ComponentCount.Y;
-					int32 Delta = FMath::Round(HitLocation.Y + (float)InitialComponentCountY / 2);
+					int32 Delta = FMath::RoundToInt(HitLocation.Y + (float)InitialComponentCountY / 2);
 					UISettings->NewLandscape_ComponentCount.Y = InitialComponentCountY - Delta;
 					UISettings->NewLandscape_ClampSize();
 					const int32 ActualDelta = UISettings->NewLandscape_ComponentCount.Y - InitialComponentCountY;
@@ -1423,7 +1423,7 @@ bool FEdModeLandscape::InputDelta( FLevelEditorViewportClient* InViewportClient,
 			case ELandscapeEdge::X_Positive_Y_Positive:
 				{
 					const int32 InitialComponentCountY = UISettings->NewLandscape_ComponentCount.Y;
-					int32 Delta = FMath::Round(HitLocation.Y - (float)InitialComponentCountY / 2);
+					int32 Delta = FMath::RoundToInt(HitLocation.Y - (float)InitialComponentCountY / 2);
 					UISettings->NewLandscape_ComponentCount.Y = InitialComponentCountY + Delta;
 					UISettings->NewLandscape_ClampSize();
 					const int32 ActualDelta = UISettings->NewLandscape_ComponentCount.Y - InitialComponentCountY;
diff --git a/Engine/Source/Editor/LandscapeEditor/Private/LandscapeEdModeBrushes.cpp b/Engine/Source/Editor/LandscapeEditor/Private/LandscapeEdModeBrushes.cpp
--- a/Engine/Source/Editor/LandscapeEditor/Private/LandscapeEdModeBrushes.cpp
+++ b/Engine/Source/Editor/LandscapeEditor/Private/LandscapeEdModeBrushes.cpp
@@ -108,10 +108,10 @@ public:
 			float Falloff = EdMode->UISettings->BrushFalloff * EdMode->UISettings->BrushRadius / ScaleXY;
 
 			// Set brush material.
-			int32 X1 = FMath::Floor(LastMousePosition.X - (Radius+Falloff));
-			int32 Y1 = FMath::Floor(LastMousePosition.Y - (Radius+Falloff));
-			int32 X2 = FMath::Ceil(LastMousePosition.X + (Radius+Falloff));
-			int32 Y2 = FMath::Ceil(LastMousePosition.Y + (Radius+Falloff));
+			int32 X1 = FMath::FloorToInt(LastMousePosition.X - (Radius+Falloff));
+			int32 Y1 = FMath::FloorToInt(LastMousePosition.Y - (Radius+Falloff));
+			int32 X2 = FMath::CeilToInt(LastMousePosition.X + (Radius+Falloff));
+			int32 Y2 = FMath::CeilToInt(LastMousePosition.Y + (Radius+Falloff));
 
 			TSet<ULandscapeComponent*> NewComponents;
 
@@ -224,10 +224,10 @@ public:
 				float MousePositionX = MousePositions[MoveIdx].PositionX;
 				float MousePositionY = MousePositions[MoveIdx].PositionY;
 
-				int32 ThisMoveX1 = FMath::Floor(MousePositionX - (Radius+Falloff));
-				int32 ThisMoveY1 = FMath::Floor(MousePositionY - (Radius+Falloff));
-				int32 ThisMoveX2 = FMath::Ceil(MousePositionX + (Radius+Falloff));
-				int32 ThisMoveY2 = FMath::Ceil(MousePositionY + (Radius+Falloff));
+				int32 ThisMoveX1 = FMath::FloorToInt(MousePositionX - (Radius+Falloff));
+				int32 ThisMoveY1 = FMath::FloorToInt(MousePositionY - (Radius+Falloff));
+				int32 ThisMoveX2 = FMath::CeilToInt(MousePositionX + (Radius+Falloff));
+				int32 ThisMoveY2 = FMath::CeilToInt(MousePositionY + (Radius+Falloff));
 				
 				if( MoveIdx == 0 )
 				{
@@ -348,8 +348,8 @@ public:
 			ULandscapeInfo* LandscapeInfo = EdMode->CurrentToolTarget.LandscapeInfo.Get();
 			if (LandscapeInfo && LandscapeInfo->ComponentSizeQuads > 0)
 			{
-				int32 X = FMath::Floor(LastMousePosition.X);
-				int32 Y = FMath::Floor(LastMousePosition.Y);
+				int32 X = FMath::FloorToInt(LastMousePosition.X);
+				int32 Y = FMath::FloorToInt(LastMousePosition.Y);
 
 				int32 ComponentIndexX = (X >= 0.f) ? X / LandscapeInfo->ComponentSizeQuads : (X + 1) / LandscapeInfo->ComponentSizeQuads - 1;
 				int32 ComponentIndexY = (Y >= 0.f) ? Y / LandscapeInfo->ComponentSizeQuads : (Y + 1) / LandscapeInfo->ComponentSizeQuads - 1;
@@ -427,8 +427,8 @@ public:
 		if( EdMode->LandscapeRenderAddCollision )
 		{
 			// Apply Brush size..
-			int32 X = FMath::Floor(LastMousePosition.X);
-			int32 Y = FMath::Floor(LastMousePosition.Y);
+			int32 X = FMath::FloorToInt(LastMousePosition.X);
+			int32 Y = FMath::FloorToInt(LastMousePosition.Y);
 
 			int32 ComponentIndexX = (X >= 0.f) ? X / LandscapeInfo->ComponentSizeQuads : (X+1) / LandscapeInfo->ComponentSizeQuads - 1;
 			int32 ComponentIndexY = (Y >= 0.f) ? Y / LandscapeInfo->ComponentSizeQuads : (Y+1) / LandscapeInfo->ComponentSizeQuads - 1;
@@ -903,9 +903,9 @@ public:
 		int32 SizeY = EdMode->UISettings->AlphaTextureSizeY;
 
 		// Bilinear interpolate the values from the alpha texture
-		int32 SampleX0 = FMath::Floor(SampleX);
+		int32 SampleX0 = FMath::FloorToInt(SampleX);
 		int32 SampleX1 = (SampleX0+1) % SizeX;
-		int32 SampleY0 = FMath::Floor(SampleY);
+		int32 SampleY0 = FMath::FloorToInt(SampleY);
 		int32 SampleY1 = (SampleY0+1) % SizeY;
 
 		const uint8* AlphaData = EdMode->UISettings->AlphaTextureData.GetData();
@@ -955,10 +955,10 @@ public:
 		int32 SizeX = EdMode->UISettings->AlphaTextureSizeX;
 		int32 SizeY = EdMode->UISettings->AlphaTextureSizeY;
 
-		X1 = FMath::Floor(LastMousePosition.X - (Radius+Falloff));
-		Y1 = FMath::Floor(LastMousePosition.Y - (Radius+Falloff));
-		X2 = FMath::Ceil(LastMousePosition.X + (Radius+Falloff));
-		Y2 = FMath::Ceil(LastMousePosition.Y + (Radius+Falloff));
+		X1 = FMath::FloorToInt(LastMousePosition.X - (Radius+Falloff));
+		Y1 = FMath::FloorToInt(LastMousePosition.Y - (Radius+Falloff));
+		X2 = FMath::CeilToInt(LastMousePosition.X + (Radius+Falloff));
+		Y2 = FMath::CeilToInt(LastMousePosition.Y + (Radius+Falloff));
 
 		for( int32 Y=Y1;Y<=Y2;Y++ )
 		{
@@ -1081,10 +1081,10 @@ public:
 		ULandscapeInfo* LandscapeInfo = EdMode->CurrentToolTarget.LandscapeInfo.Get();
 		if( OldMousePosition.IsZero() )
 		{
-			X1 = FMath::Floor(LastMousePosition.X);
-			Y1 = FMath::Floor(LastMousePosition.Y);
-			X2 = FMath::Ceil(LastMousePosition.X);
-			Y2 = FMath::Ceil(LastMousePosition.Y);
+			X1 = FMath::FloorToInt(LastMousePosition.X);
+			Y1 = FMath::FloorToInt(LastMousePosition.Y);
+			X2 = FMath::CeilToInt(LastMousePosition.X);
+			Y2 = FMath::CeilToInt(LastMousePosition.Y);
 			OldMousePosition = LastMousePosition;
 			LastMouseAngle = 0.f;
 			LastMouseSampleTime = FPlatformTime::Seconds();
@@ -1098,10 +1098,10 @@ public:
 			float MaxSize = 2.f * FMath::Sqrt( FMath::Square(Radius) / 2.f );
 			float AlphaBrushScale = MaxSize / (float)FMath::Max<int32>(SizeX,SizeY);
 
-			X1 = FMath::Floor(LastMousePosition.X - Radius);
-			Y1 = FMath::Floor(LastMousePosition.Y - Radius);
-			X2 = FMath::Ceil(LastMousePosition.X + Radius);
-			Y2 = FMath::Ceil(LastMousePosition.Y + Radius);
+			X1 = FMath::FloorToInt(LastMousePosition.X - Radius);
+			Y1 = FMath::FloorToInt(LastMousePosition.Y - Radius);
+			X2 = FMath::CeilToInt(LastMousePosition.X + Radius);
+			Y2 = FMath::CeilToInt(LastMousePosition.Y + Radius);
 
 			for( int32 Y=Y1;Y<=Y2;Y++ )
 			{
diff --git a/Engine/Source/Editor/LandscapeEditor/Private/LandscapeEdModeComponentTools.cpp b/Engine/Source/Editor/LandscapeEditor/Private/LandscapeEdModeComponentTools.cpp
--- a/Engine/Source/Editor/LandscapeEditor/Private/LandscapeEdModeComponentTools.cpp
+++ b/Engine/Source/Editor/LandscapeEditor/Private/LandscapeEdModeComponentTools.cpp
@@ -53,8 +53,8 @@ public:
 				if (!bInitializedComponentInvert)
 				{
 					// Get the component under the mouse location. Copied from FLandscapeBrushComponent::ApplyBrush()
-					const int32 MouseX = FMath::Round(MousePositions[0].PositionX);
-					const int32 MouseY = FMath::Round(MousePositions[0].PositionY);
+					const int32 MouseX = FMath::RoundToInt(MousePositions[0].PositionX);
+					const int32 MouseY = FMath::RoundToInt(MousePositions[0].PositionY);
 					const int32 MouseComponentIndexX = (MouseX >= 0.f) ? MouseX / LandscapeInfo->ComponentSizeQuads : (MouseX+1) / LandscapeInfo->ComponentSizeQuads - 1;
 					const int32 MouseComponentIndexY = (MouseY >= 0.f) ? MouseY / LandscapeInfo->ComponentSizeQuads : (MouseY+1) / LandscapeInfo->ComponentSizeQuads - 1;
 					ULandscapeComponent* MouseComponent = LandscapeInfo->XYtoComponentMap.FindRef(FIntPoint(MouseComponentIndexX, MouseComponentIndexY));
@@ -139,7 +139,7 @@ public:
 							LandscapeInfo->SelectedRegion.Add(Key, FMath::Min(Value + PaintValue, 1.0f));
 						}
 
-						Data[(Key.X-X1) + (Key.Y-Y1)*(1+X2-X1)] = FMath::Clamp<int32>(FMath::Round(LandscapeInfo->SelectedRegion.FindRef(Key) * 255), 0, 255);
+						Data[(Key.X-X1) + (Key.Y-Y1)*(1+X2-X1)] = FMath::Clamp<int32>(FMath::RoundToInt(LandscapeInfo->SelectedRegion.FindRef(Key) * 255), 0, 255);
 					}
 				}
 
@@ -245,7 +245,7 @@ public:
 						Value = 255; // Just on and off for visibility, for masking...
 					}
 
-					Data[(X-X1) + (Y-Y1)*(1+X2-X1)] = FMath::Clamp<int32>(FMath::Round(Value), 0, 255);
+					Data[(X-X1) + (Y-Y1)*(1+X2-X1)] = FMath::Clamp<int32>(FMath::RoundToInt(Value), 0, 255);
 				}
 			}
 
@@ -1141,8 +1141,8 @@ public:
 			Gizmo->CachedScaleXY = ScaleXY;
 
 			// Rasterize Gizmo regions
-			int32 SizeX = FMath::Ceil(Width  / ScaleXY);
-			int32 SizeY = FMath::Ceil(Height / ScaleXY);
+			int32 SizeX = FMath::CeilToInt(Width  / ScaleXY);
+			int32 SizeY = FMath::CeilToInt(Height / ScaleXY);
 
 			const float W = (Width - ScaleXY) / (2 * ScaleXY);
 			const float H = (Height - ScaleXY) / (2 * ScaleXY);
@@ -1163,8 +1163,8 @@ public:
 				for (int32 X = 0; X < SizeX; ++X)
 				{
 					FVector LandscapeLocal = GizmoLocalToLandscape.TransformPosition(FVector( -W + X, -H + Y, 0 ));
-					int32 LX = FMath::Floor(LandscapeLocal.X);
-					int32 LY = FMath::Floor(LandscapeLocal.Y);
+					int32 LX = FMath::FloorToInt(LandscapeLocal.X);
+					int32 LY = FMath::FloorToInt(LandscapeLocal.Y);
 
 					{
 						for (int32 i = -1; (!bApplyToAll && i < 0) || i < LayerNum; ++i )
@@ -1529,8 +1529,8 @@ public:
 					GizmoLocal.X *= ScaleX * SignX;
 					GizmoLocal.Y *= ScaleY * SignY;
 
-					int32 LX = FMath::Floor(GizmoLocal.X);
-					int32 LY = FMath::Floor(GizmoLocal.Y);
+					int32 LX = FMath::FloorToInt(GizmoLocal.X);
+					int32 LY = FMath::FloorToInt(GizmoLocal.Y);
 
 					float FracX = GizmoLocal.X - LX;
 					float FracY = GizmoLocal.Y - LY;
diff --git a/Engine/Source/Editor/LandscapeEditor/Private/LandscapeEdModePaintTools.cpp b/Engine/Source/Editor/LandscapeEditor/Private/LandscapeEdModePaintTools.cpp
--- a/Engine/Source/Editor/LandscapeEditor/Private/LandscapeEdModePaintTools.cpp
+++ b/Engine/Source/Editor/LandscapeEditor/Private/LandscapeEdModePaintTools.cpp
@@ -202,7 +202,7 @@ public:
 				ALandscape::UnpackKey(It.Key(), X, Y);
 				float Weight = It.Value();
 
-				AveragePoint += FVector( (float)X * Weight, (float)Y * Weight, (float)DATA_AT(OriginalData,FMath::Floor(X),FMath::Floor(Y)) * Weight );
+				AveragePoint += FVector( (float)X * Weight, (float)Y * Weight, (float)DATA_AT(OriginalData,FMath::FloorToInt(X),FMath::FloorToInt(Y)) * Weight );
 
 				FVector SampleNormal = DATA_AT(Normals,X,Y);
 				AverageNormal += SampleNormal * Weight;
@@ -270,22 +270,22 @@ public:
 
 				if( bInvert )
 				{
-					CurrentValue = ToolTarget::CacheClass::ClampValue( FMath::Min<int32>(FMath::Round(PaintValue), CurrentValue) );
+					CurrentValue = ToolTarget::CacheClass::ClampValue( FMath::Min<int32>(FMath::RoundToInt(PaintValue), CurrentValue) );
 				}
 				else
 				{
-					CurrentValue = ToolTarget::CacheClass::ClampValue( FMath::Max<int32>(FMath::Round(PaintValue), CurrentValue) );
+					CurrentValue = ToolTarget::CacheClass::ClampValue( FMath::Max<int32>(FMath::RoundToInt(PaintValue), CurrentValue) );
 				}
 			}
 			else
 			{
 				if( bInvert )
 				{
-					CurrentValue = ToolTarget::CacheClass::ClampValue( FMath::Min<int32>(SourceValue - FMath::Round(PaintAmount), CurrentValue) );
+					CurrentValue = ToolTarget::CacheClass::ClampValue( FMath::Min<int32>(SourceValue - FMath::RoundToInt(PaintAmount), CurrentValue) );
 				}
 				else
 				{
-					CurrentValue = ToolTarget::CacheClass::ClampValue( FMath::Max<int32>(SourceValue + FMath::Round(PaintAmount), CurrentValue) );
+					CurrentValue = ToolTarget::CacheClass::ClampValue( FMath::Max<int32>(SourceValue + FMath::RoundToInt(PaintAmount), CurrentValue) );
 				}
 			}
 		}
@@ -445,8 +445,8 @@ public:
 			bInitializedFlattenHeight = false;
 			float FlattenX = MousePositions[0].PositionX;
 			float FlattenY = MousePositions[0].PositionY;
-			int32 FlattenHeightX = FMath::Floor(FlattenX);
-			int32 FlattenHeightY = FMath::Floor(FlattenY);
+			int32 FlattenHeightX = FMath::FloorToInt(FlattenX);
+			int32 FlattenHeightY = FMath::FloorToInt(FlattenY);
 
 			this->Cache.CacheData(FlattenHeightX,FlattenHeightY,FlattenHeightX+1,FlattenHeightY+1);
 			float HeightValue = this->Cache.GetValue(FlattenX, FlattenY);
@@ -583,9 +583,9 @@ public:
 
 			const FTransform LocalToWorld = this->EdMode->CurrentToolTarget.LandscapeInfo->GetLandscapeProxy()->ActorToWorld();
 			FVector Origin;
-			Origin.X = FMath::Round(LastMousePosition.X);
-			Origin.Y = FMath::Round(LastMousePosition.Y);
-			Origin.Z = (FMath::Round((this->EdMode->UISettings->FlattenTarget / LocalToWorld.GetScale3D().Z - LocalToWorld.GetTranslation().Z) * LANDSCAPE_INV_ZSCALE) - 0.1f) * LANDSCAPE_ZSCALE;
+			Origin.X = FMath::RoundToFloat(LastMousePosition.X);
+			Origin.Y = FMath::RoundToFloat(LastMousePosition.Y);
+			Origin.Z = (FMath::RoundToFloat((this->EdMode->UISettings->FlattenTarget / LocalToWorld.GetScale3D().Z - LocalToWorld.GetTranslation().Z) * LANDSCAPE_INV_ZSCALE) - 0.1f) * LANDSCAPE_ZSCALE;
 			MeshComponent->SetRelativeLocation(Origin, false);
 		}
 
@@ -683,7 +683,7 @@ public:
 						}
 						break;
 					}
-					Data[(X-X1) + (Y-Y1)*(1+X2-X1)] = ToolTarget::CacheClass::ClampValue( FMath::Round(FMath::Lerp( OriginalValue, DestValue, It.Value() * UISettings->ToolStrength * Pressure)) );
+					Data[(X-X1) + (Y-Y1)*(1+X2-X1)] = ToolTarget::CacheClass::ClampValue( FMath::RoundToInt(FMath::Lerp( OriginalValue, DestValue, It.Value() * UISettings->ToolStrength * Pressure)) );
 				}
 				else
 				{
diff --git a/Engine/Source/Editor/LandscapeEditor/Private/LandscapeEdModeRampTool.cpp b/Engine/Source/Editor/LandscapeEditor/Private/LandscapeEdModeRampTool.cpp
--- a/Engine/Source/Editor/LandscapeEditor/Private/LandscapeEdModeRampTool.cpp
+++ b/Engine/Source/Editor/LandscapeEditor/Private/LandscapeEdModeRampTool.cpp
@@ -209,8 +209,8 @@ public:
 		{
 			if (InKey == EKeys::End && InEvent == IE_Pressed)
 			{
-				const int32 MinX = FMath::Floor(Points[SelectedPoint].X);
-				const int32 MinY = FMath::Floor(Points[SelectedPoint].Y);
+				const int32 MinX = FMath::FloorToInt(Points[SelectedPoint].X);
+				const int32 MinY = FMath::FloorToInt(Points[SelectedPoint].Y);
 				const int32 MaxX = MinX + 1;
 				const int32 MaxY = MinY + 1;
 
@@ -236,7 +236,7 @@ public:
 				checkSlow(ValidMaxX == MaxX);
 				checkSlow(ValidMaxY == MaxY);
 
-				Points[SelectedPoint].Z = (FMath::BiLerp<float>(Data[0], Data[1], Data[2], Data[3], FMath::Fractional(Points[SelectedPoint].X), FMath::Fractional(Points[SelectedPoint].Y)) - LandscapeDataAccess::MidValue) * LANDSCAPE_ZSCALE;
+				Points[SelectedPoint].Z = (FMath::BiLerp<float>(Data[0], Data[1], Data[2], Data[3], FMath::Frac(Points[SelectedPoint].X), FMath::Frac(Points[SelectedPoint].Y)) - LandscapeDataAccess::MidValue) * LANDSCAPE_ZSCALE;
 
 				return true;
 			}
@@ -262,7 +262,7 @@ public:
 				NewValue = FMath::Max(NewValue, OldValue + 1.f);
 			}
 
-			NewValue = FMath::Round(FMath::Clamp(NewValue, SliderMin, SliderMax));
+			NewValue = FMath::RoundToFloat(FMath::Clamp(NewValue, SliderMin, SliderMax));
 
 			EdMode->UISettings->RampWidth = NewValue;
 
@@ -448,10 +448,10 @@ public:
 		Heights[0] = Points[0].Z * LANDSCAPE_INV_ZSCALE + LandscapeDataAccess::MidValue;
 		Heights[1] = Points[1].Z * LANDSCAPE_INV_ZSCALE + LandscapeDataAccess::MidValue;
 
-		int32 MinX = FMath::Ceil( FMath::Min(FMath::Min(OuterVerts[0][0].X, OuterVerts[0][1].X),FMath::Min(OuterVerts[1][0].X, OuterVerts[1][1].X))) - 1; // +/- 1 to make sure we have enough data for calculating correct normals
-		int32 MinY = FMath::Ceil( FMath::Min(FMath::Min(OuterVerts[0][0].Y, OuterVerts[0][1].Y),FMath::Min(OuterVerts[1][0].Y, OuterVerts[1][1].Y))) - 1;
-		int32 MaxX = FMath::Floor(FMath::Max(FMath::Max(OuterVerts[0][0].X, OuterVerts[0][1].X),FMath::Max(OuterVerts[1][0].X, OuterVerts[1][1].X))) + 1;
-		int32 MaxY = FMath::Floor(FMath::Max(FMath::Max(OuterVerts[0][0].Y, OuterVerts[0][1].Y),FMath::Max(OuterVerts[1][0].Y, OuterVerts[1][1].Y))) + 1;
+		int32 MinX = FMath::CeilToInt( FMath::Min(FMath::Min(OuterVerts[0][0].X, OuterVerts[0][1].X),FMath::Min(OuterVerts[1][0].X, OuterVerts[1][1].X))) - 1; // +/- 1 to make sure we have enough data for calculating correct normals
+		int32 MinY = FMath::CeilToInt( FMath::Min(FMath::Min(OuterVerts[0][0].Y, OuterVerts[0][1].Y),FMath::Min(OuterVerts[1][0].Y, OuterVerts[1][1].Y))) - 1;
+		int32 MaxX = FMath::FloorToInt(FMath::Max(FMath::Max(OuterVerts[0][0].X, OuterVerts[0][1].X),FMath::Max(OuterVerts[1][0].X, OuterVerts[1][1].X))) + 1;
+		int32 MaxY = FMath::FloorToInt(FMath::Max(FMath::Max(OuterVerts[0][0].Y, OuterVerts[0][1].Y),FMath::Max(OuterVerts[1][0].Y, OuterVerts[1][1].Y))) + 1;
 
 		FLandscapeEditDataInterface LandscapeEdit(EdMode->CurrentToolTarget.LandscapeInfo.Get());
 
diff --git a/Engine/Source/Editor/LandscapeEditor/Private/LandscapeEdModeTools.h b/Engine/Source/Editor/LandscapeEditor/Private/LandscapeEdModeTools.h
--- a/Engine/Source/Editor/LandscapeEditor/Private/LandscapeEdModeTools.h
+++ b/Engine/Source/Editor/LandscapeEditor/Private/LandscapeEdModeTools.h
@@ -111,8 +111,8 @@ private:
 
 	float PerlinNoise2D(float X,float Y) const
 	{
-		int32		TruncX = FMath::Trunc(X),
-			TruncY = FMath::Trunc(Y),
+		int32		TruncX = FMath::TruncToInt(X),
+			TruncY = FMath::TruncToInt(Y),
 			IntX = TruncX & 255,
 			IntY = TruncY & 255;
 		float	FracX = X - TruncX,
@@ -374,8 +374,8 @@ struct TLandscapeEditCache
 
 	float GetValue(float LandscapeX, float LandscapeY)
 	{
-		int32 X = FMath::Floor(LandscapeX);
-		int32 Y = FMath::Floor(LandscapeY);
+		int32 X = FMath::FloorToInt(LandscapeX);
+		int32 Y = FMath::FloorToInt(LandscapeY);
 		AccessorType* P00 = CachedData.Find(ALandscape::MakeKey(X, Y));
 		AccessorType* P10 = CachedData.Find(ALandscape::MakeKey(X+1, Y));
 		AccessorType* P01 = CachedData.Find(ALandscape::MakeKey(X, Y+1));
diff --git a/Engine/Source/Editor/LandscapeEditor/Private/LandscapeEdModeXYOffset.cpp b/Engine/Source/Editor/LandscapeEditor/Private/LandscapeEdModeXYOffset.cpp
--- a/Engine/Source/Editor/LandscapeEditor/Private/LandscapeEdModeXYOffset.cpp
+++ b/Engine/Source/Editor/LandscapeEditor/Private/LandscapeEdModeXYOffset.cpp
@@ -32,8 +32,8 @@ namespace
 		float Height[4];
 		for (int32 Idx = 0; Idx < 4; ++Idx)
 		{
-			int32 XX = FMath::Clamp(FMath::Floor(X+XYOffset.X + XOffsets[Idx]), MinX, MaxX);
-			int32 YY = FMath::Clamp(FMath::Floor(Y+XYOffset.Y + YOffsets[Idx]), MinY, MaxY);
+			int32 XX = FMath::Clamp(FMath::FloorToInt(X+XYOffset.X + XOffsets[Idx]), MinX, MaxX);
+			int32 YY = FMath::Clamp(FMath::FloorToInt(Y+XYOffset.Y + YOffsets[Idx]), MinY, MaxY);
 			Height[Idx] = XYOffsetVectorData[ XX - MinX + (YY - MinY) * (MaxX-MinX+1) ].Z;
 		}
 		float FracX = FMath::Fractional(X+XYOffset.X);
@@ -242,7 +242,7 @@ public:
 				int32 IterNum = 0;
 				float TotalHeightError = 0.f;
 
-				while( NewXYOffset[ (Y - Y1) * (X2-X1+1) ].Y + Y > XYOffsetVectorData[(FPlatformMath::Floor(YOffset) - Y1) * (X2-X1+1) ].Y + FPlatformMath::Floor(YOffset) )
+				while( NewXYOffset[ (Y - Y1) * (X2-X1+1) ].Y + Y > XYOffsetVectorData[(FPlatformMath::FloorToInt(YOffset) - Y1) * (X2-X1+1) ].Y + FPlatformMath::FloorToInt(YOffset) )
 				{
 					YOffset = YOffset+1.f;
 					if (YOffset >= Y2)
@@ -272,9 +272,9 @@ public:
 								P[Idx] = FVector( GetWorldPos( LocalToWorld, FVector2D(XX - Comp->GetSectionBase().X, Y - Comp->GetSectionBase().Y), NewXYOffset[ XX - X1 + (Y - Y1) * (X2-X1+1) ] ) );
 							}
 
-							int32 YY0 = FMath::Clamp<int32>(FMath::Floor(YOffset-1), Y1, Y2);
-							int32 YY1 = FMath::Clamp<int32>(FMath::Floor(YOffset), Y1, Y2);
-							int32 YY2 = FMath::Clamp<int32>(FMath::Floor(1+YOffset), Y1, Y2);
+							int32 YY0 = FMath::Clamp<int32>(FMath::FloorToInt(YOffset-1), Y1, Y2);
+							int32 YY1 = FMath::Clamp<int32>(FMath::FloorToInt(YOffset), Y1, Y2);
+							int32 YY2 = FMath::Clamp<int32>(FMath::FloorToInt(1+YOffset), Y1, Y2);
 							// Search for valid YOffset...
 							for (int32 Idx = 2; Idx < 4; ++Idx)
 							{
@@ -349,8 +349,8 @@ public:
 					RemainQuads -= QuadY[Y - Y1];
 					for (int32 X = MinX[Y-Y1]; X < MaxX[Y-Y1]; ++X)
 					{
-						int32 YY1 = FMath::Clamp<int32>(FMath::Floor(YOffset), Y1, Y2);
-						int32 YY2 = FMath::Clamp<int32>(FMath::Floor(1+YOffset), Y1, Y2);
+						int32 YY1 = FMath::Clamp<int32>(FMath::FloorToInt(YOffset), Y1, Y2);
+						int32 YY2 = FMath::Clamp<int32>(FMath::FloorToInt(1+YOffset), Y1, Y2);
 						FVector P1 = XYOffsetVectorData[ X - X1 + (YY1 - Y1) * (X2-X1+1) ];
 						//P1.X = X+P1.X;
 						P1.Y = YY1+P1.Y;
@@ -389,7 +389,7 @@ public:
 				int32 IterNum = 0;
 				float TotalHeightError = 0.f;
 
-				while( NewXYOffset[X - X1].X + X > NewYOffsets[FMath::Floor(XOffset) - X1].X + FMath::Floor(XOffset) )
+				while( NewXYOffset[X - X1].X + X > NewYOffsets[FMath::FloorToInt(XOffset) - X1].X + FMath::FloorToFloat(XOffset) )
 				{
 					XOffset = XOffset+1.f;
 					if (XOffset >= X2)
@@ -417,9 +417,9 @@ public:
 								P[Idx] = FVector( GetWorldPos( LocalToWorld, FVector2D(X - Comp->GetSectionBase().X, YY - Comp->GetSectionBase().Y), NewXYOffset[ X - X1 + (YY - Y1) * (X2-X1+1) ] ) );
 							}
 
-							int32 XX0 = FMath::Clamp<int32>(FMath::Floor(XOffset-1), X1, X2);
-							int32 XX1 = FMath::Clamp<int32>(FMath::Floor(XOffset), X1, X2);
-							int32 XX2 = FMath::Clamp<int32>(FMath::Floor(1+XOffset), X1, X2);
+							int32 XX0 = FMath::Clamp<int32>(FMath::FloorToInt(XOffset-1), X1, X2);
+							int32 XX1 = FMath::Clamp<int32>(FMath::FloorToInt(XOffset), X1, X2);
+							int32 XX2 = FMath::Clamp<int32>(FMath::FloorToInt(1+XOffset), X1, X2);
 
 							// Search for valid YOffset...
 							for (int32 Idx = 1; Idx < 4; Idx+=2)
@@ -496,8 +496,8 @@ public:
 
 					for (int32 Y = MinY[X-X1]; Y < MaxY[X-X1]; ++Y)
 					{
-						int32 XX1 = FMath::Clamp<int32>(FMath::Floor(XOffset), X1, X2);
-						int32 XX2 = FMath::Clamp<int32>(FMath::Floor(1+XOffset), X1, X2);
+						int32 XX1 = FMath::Clamp<int32>(FMath::FloorToInt(XOffset), X1, X2);
+						int32 XX2 = FMath::Clamp<int32>(FMath::FloorToInt(1+XOffset), X1, X2);
 						FVector P1 = NewYOffsets[ XX1 - X1 + (Y - Y1) * (X2-X1+1) ];
 						P1.X = XX1+P1.X;
 						FVector P2 = NewYOffsets[ XX2 - X1 + (Y - Y1) * (X2-X1+1) ];
diff --git a/Engine/Source/Editor/LandscapeEditor/Private/LandscapeEditorDetailCustomization_NewLandscape.cpp b/Engine/Source/Editor/LandscapeEditor/Private/LandscapeEditorDetailCustomization_NewLandscape.cpp
--- a/Engine/Source/Editor/LandscapeEditor/Private/LandscapeEditorDetailCustomization_NewLandscape.cpp
+++ b/Engine/Source/Editor/LandscapeEditor/Private/LandscapeEditorDetailCustomization_NewLandscape.cpp
@@ -654,7 +654,7 @@ TOptional<int32> FLandscapeEditorDetailCustomization_NewLandscape::GetMaxLandsca
 	{
 		// Max size is either whole components below 8192 verts, or 32 components
 		const int32 QuadsPerComponent = (LandscapeEdMode->UISettings->NewLandscape_SectionsPerComponent * LandscapeEdMode->UISettings->NewLandscape_QuadsPerSection);
-		//return (float)(FMath::Min(32, FMath::Floor(8191 / QuadsPerComponent)) * QuadsPerComponent);
+		//return (float)(FMath::Min(32, FMath::FloorToInt(8191 / QuadsPerComponent)) * QuadsPerComponent);
 		return (8191 / QuadsPerComponent) * QuadsPerComponent + 1;
 	}
 
@@ -885,8 +885,8 @@ FReply FLandscapeEditorDetailCustomization_NewLandscape::OnFillWorldButtonClicke
 		NewLandscapeLocation.Y = 0;
 
 		const int32 QuadsPerComponent = LandscapeEdMode->UISettings->NewLandscape_SectionsPerComponent * LandscapeEdMode->UISettings->NewLandscape_QuadsPerSection;
-		LandscapeEdMode->UISettings->NewLandscape_ComponentCount.X = FMath::Ceil(WORLD_MAX / QuadsPerComponent / LandscapeEdMode->UISettings->NewLandscape_Scale.X);
-		LandscapeEdMode->UISettings->NewLandscape_ComponentCount.Y = FMath::Ceil(WORLD_MAX / QuadsPerComponent / LandscapeEdMode->UISettings->NewLandscape_Scale.Y);
+		LandscapeEdMode->UISettings->NewLandscape_ComponentCount.X = FMath::CeilToInt(WORLD_MAX / QuadsPerComponent / LandscapeEdMode->UISettings->NewLandscape_Scale.X);
+		LandscapeEdMode->UISettings->NewLandscape_ComponentCount.Y = FMath::CeilToInt(WORLD_MAX / QuadsPerComponent / LandscapeEdMode->UISettings->NewLandscape_Scale.Y);
 		LandscapeEdMode->UISettings->NewLandscape_ClampSize();
 	}
 
@@ -1075,7 +1075,7 @@ void FLandscapeEditorDetailCustomization_NewLandscape::OnImportHeightmapFilename
 					// Find all possible heightmap sizes, up to 8192 width/height
 					int32 InsertIndex = 0;
 					const int32 MinWidth = FMath::DivideAndRoundUp(ImportFileSize, (int64)8192);
-					const int32 MaxWidth = FMath::Trunc(FMath::Sqrt(ImportFileSize));
+					const int32 MaxWidth = FMath::TruncToInt(FMath::Sqrt(ImportFileSize));
 					for (int32 Width = MinWidth; Width <= MaxWidth; Width++)
 					{
 						if (ImportFileSize % Width == 0)
diff --git a/Engine/Source/Editor/LandscapeEditor/Private/LandscapeEditorObject.cpp b/Engine/Source/Editor/LandscapeEditor/Private/LandscapeEditorObject.cpp
--- a/Engine/Source/Editor/LandscapeEditor/Private/LandscapeEditorObject.cpp
+++ b/Engine/Source/Editor/LandscapeEditor/Private/LandscapeEditorObject.cpp
@@ -398,7 +398,7 @@ void ULandscapeEditorObject::GuessGizmoImportSize()
 
 			// Guess dimensions from filesize
 			// Keep searching for the most squarelike size
-			for (int32 w = FMath::Trunc(FMath::Sqrt((float)GizmoFileSize)); w > 0; w--)
+			for (int32 w = FMath::TruncToInt(FMath::Sqrt((float)GizmoFileSize)); w > 0; w--)
 			{
 				if (GizmoFileSize % w == 0)
 				{
diff --git a/Engine/Source/Editor/LandscapeEditor/Public/LandscapeEdMode.h b/Engine/Source/Editor/LandscapeEditor/Public/LandscapeEdMode.h
--- a/Engine/Source/Editor/LandscapeEditor/Public/LandscapeEdMode.h
+++ b/Engine/Source/Editor/LandscapeEditor/Public/LandscapeEdMode.h
@@ -730,10 +730,10 @@ namespace LandscapeEditorUtils
 			{
 				const float OldY = Y * YScale;
 				const float OldX = X * XScale;
-				const int32 X0 = FMath::Floor(OldX);
-				const int32 X1 = FMath::Min(FMath::Floor(OldX) + 1, OldWidth - 1);
-				const int32 Y0 = FMath::Floor(OldY);
-				const int32 Y1 = FMath::Min(FMath::Floor(OldY) + 1, OldHeight - 1);
+				const int32 X0 = FMath::FloorToInt(OldX);
+				const int32 X1 = FMath::Min(FMath::FloorToInt(OldX) + 1, OldWidth - 1);
+				const int32 Y0 = FMath::FloorToInt(OldY);
+				const int32 Y1 = FMath::Min(FMath::FloorToInt(OldY) + 1, OldHeight - 1);
 				const T& Original00 = Data[Y0 * OldWidth + X0];
 				const T& Original10 = Data[Y0 * OldWidth + X1];
 				const T& Original01 = Data[Y1 * OldWidth + X0];
diff --git a/Engine/Source/Editor/LevelEditor/Private/HighresScreenshotUI.h b/Engine/Source/Editor/LevelEditor/Private/HighresScreenshotUI.h
--- a/Engine/Source/Editor/LevelEditor/Private/HighresScreenshotUI.h
+++ b/Engine/Source/Editor/LevelEditor/Private/HighresScreenshotUI.h
@@ -73,7 +73,7 @@ private:
 	void OnResolutionMultiplierSliderChanged( float NewValue )
 	{
 		Config.ResolutionMultiplierScale = NewValue;
-		Config.ResolutionMultiplier = FMath::Round(FMath::Lerp(FHighResScreenshotConfig::MinResolutionMultipler, FHighResScreenshotConfig::MaxResolutionMultipler, NewValue));
+		Config.ResolutionMultiplier = FMath::RoundToFloat(FMath::Lerp(FHighResScreenshotConfig::MinResolutionMultipler, FHighResScreenshotConfig::MaxResolutionMultipler, NewValue));
 	}
 
 	void OnMaskEnabledChanged( ESlateCheckBoxState::Type NewValue )
diff --git a/Engine/Source/Editor/LevelEditor/Private/LevelEditor.cpp b/Engine/Source/Editor/LevelEditor/Private/LevelEditor.cpp
--- a/Engine/Source/Editor/LevelEditor/Private/LevelEditor.cpp
+++ b/Engine/Source/Editor/LevelEditor/Private/LevelEditor.cpp
@@ -375,7 +375,7 @@ void FLevelEditorModule::StartImmersivePlayInEditorSession()
 				ActiveLevelViewport->MakeImmersive( bWantImmersive, bAllowAnimation );
 				FVector2D WindowSize = Window->GetSizeInScreen();
 				// Set the initial size of the viewport to be the size of the window. This must be done because Slate has not ticked yet so the viewport will have no initial size
-				ActiveLevelViewport->GetActiveViewport()->SetInitialSize( FIntPoint( FMath::Trunc( WindowSize.X ), FMath::Trunc( WindowSize.Y ) ) );
+				ActiveLevelViewport->GetActiveViewport()->SetInitialSize( FIntPoint( FMath::TruncToInt( WindowSize.X ), FMath::TruncToInt( WindowSize.Y ) ) );
 			}
 
 			// Launch PIE
diff --git a/Engine/Source/Editor/LevelEditor/Private/LevelEditorActions.cpp b/Engine/Source/Editor/LevelEditor/Private/LevelEditorActions.cpp
--- a/Engine/Source/Editor/LevelEditor/Private/LevelEditorActions.cpp
+++ b/Engine/Source/Editor/LevelEditor/Private/LevelEditorActions.cpp
@@ -794,7 +794,7 @@ void FLevelEditorActionCallbacks::SetLightingResolutionMaxBSPs( float Value )
 
 int32 FLevelEditorActionCallbacks::GetLightingResolutionRatio()
 {
-	return FMath::Round(FLightmapResRatioAdjustSettings::Get().Ratio * 100.0f);
+	return FMath::RoundToInt(FLightmapResRatioAdjustSettings::Get().Ratio * 100.0f);
 }
 
 void FLevelEditorActionCallbacks::SetLightingResolutionRatio( int32 Value )
diff --git a/Engine/Source/Editor/LevelEditor/Private/LightmapResRatioAdjust.cpp b/Engine/Source/Editor/LevelEditor/Private/LightmapResRatioAdjust.cpp
--- a/Engine/Source/Editor/LevelEditor/Private/LightmapResRatioAdjust.cpp
+++ b/Engine/Source/Editor/LevelEditor/Private/LightmapResRatioAdjust.cpp
@@ -170,7 +170,7 @@ bool FLightmapResRatioAdjustSettings::ApplyRatioAdjustment()
 				{
 					SMComp->Modify();
 					float AdjustedResolution = (float)CurrentResolution * (Settings.Ratio);
-					int32 NewResolution = FMath::Trunc(AdjustedResolution);
+					int32 NewResolution = FMath::TruncToInt(AdjustedResolution);
 					NewResolution = FMath::Max(NewResolution + 3 & ~3,4);
 					SMComp->SetStaticLightingMapping(true, NewResolution);
 					SMComp->InvalidateLightingCache();
@@ -220,7 +220,7 @@ bool FLightmapResRatioAdjustSettings::ApplyRatioAdjustment()
 					if (bConvertIt == true)
 					{
 						Model->ModifySurf( SurfaceIndex, true );
-						int32 NewResolution = FMath::Trunc(Surf.LightMapScale * Scalar);
+						int32 NewResolution = FMath::TruncToInt(Surf.LightMapScale * Scalar);
 						NewResolution = FMath::Max(NewResolution + 3 & ~3,4);
 						Surf.LightMapScale = (float)NewResolution;
 						if (Surf.Actor != NULL)
