Fixed some windows not gaining focus when opened

On mac, window activation is not guaranteed syncronously when a window is created, thus any code that opened a window, then immediately opened some focus-sensitive operation, may fail.

I've fixed up three of the most common cases (SAssetView::ItemScrolledIntoView, STextEntryPopup and STextComboPopup) so they register an active tick to focus the widget when the window gains focus (optionally controllable from the .AutoFocus() argument)

This addresses UE-12929.

diff --git a/Engine/Source/Editor/Cascade/Private/Cascade.cpp b/Engine/Source/Editor/Cascade/Private/Cascade.cpp
--- a/Engine/Source/Editor/Cascade/Private/Cascade.cpp
+++ b/Engine/Source/Editor/Cascade/Private/Cascade.cpp
@@ -3704,8 +3704,6 @@ void FCascade::OnSetMotionRadius()
 		FSlateApplication::Get().GetCursorPos(),
 		FPopupTransitionEffect(FPopupTransitionEffect::TypeInPopup)
 		);
-
-	TextEntry->FocusDefaultWidget();
 }
 
 void FCascade::OnViewMode(EViewModeIndex ViewMode)
@@ -3927,8 +3925,6 @@ void FCascade::OnToggleWireframeSphere()
 			FSlateApplication::Get().GetCursorPos(),
 			FPopupTransitionEffect(FPopupTransitionEffect::TypeInPopup)
 			);
-
-		TextEntry->FocusDefaultWidget();
 	}
 	else
 	{
@@ -4820,8 +4816,6 @@ void FCascade::OnRenameEmitter()
 		FSlateApplication::Get().GetCursorPos(),
 		FPopupTransitionEffect(FPopupTransitionEffect::TypeInPopup)
 		);
-
-	TextEntry->FocusDefaultWidget();
 }
 
 void FCascade::OnDuplicateEmitter(bool bIsShared)
diff --git a/Engine/Source/Editor/ContentBrowser/Private/SAssetView.cpp b/Engine/Source/Editor/ContentBrowser/Private/SAssetView.cpp
--- a/Engine/Source/Editor/ContentBrowser/Private/SAssetView.cpp
+++ b/Engine/Source/Editor/ContentBrowser/Private/SAssetView.cpp
@@ -942,6 +942,23 @@ void SAssetView::Tick( const FGeometry& AllottedGeometry, const double InCurrent
 	{
 		ResetQuickJump();
 	}
+
+	TSharedPtr<FAssetViewItem> AssetAwaitingRename = AwaitingRename.Pin();
+	if (AssetAwaitingRename.IsValid())
+	{
+		TSharedPtr<SWindow> OwnerWindow = FSlateApplication::Get().FindWidgetWindow(AsShared());
+		if (!OwnerWindow.IsValid())
+		{
+			AssetAwaitingRename->bRenameWhenScrolledIntoview = false;
+			AwaitingRename = nullptr;
+		}
+		else if (FSlateApplication::Get().HasFocusedDescendants(OwnerWindow.ToSharedRef()))
+		{
+			AssetAwaitingRename->RenamedRequestEvent.ExecuteIfBound();
+			AssetAwaitingRename->bRenameWhenScrolledIntoview = false;
+			AwaitingRename = nullptr;
+		}
+	}
 }
 
 void SAssetView::CalculateFillScale( const FGeometry& AllottedGeometry )
@@ -3081,12 +3098,7 @@ void SAssetView::ItemScrolledIntoView(TSharedPtr<struct FAssetViewItem> AssetIte
 			OwnerWindow->BringToFront();
 		}
 
-		if ( Widget.IsValid() && Widget->GetContent().IsValid() )
-		{
-			AssetItem->RenamedRequestEvent.ExecuteIfBound();
-		}
-
-		AssetItem->bRenameWhenScrolledIntoview = false;
+		AwaitingRename = AssetItem;
 	}
 }
 
diff --git a/Engine/Source/Editor/ContentBrowser/Private/SAssetView.h b/Engine/Source/Editor/ContentBrowser/Private/SAssetView.h
--- a/Engine/Source/Editor/ContentBrowser/Private/SAssetView.h
+++ b/Engine/Source/Editor/ContentBrowser/Private/SAssetView.h
@@ -691,6 +691,9 @@ private:
 	double LastSortTime;
 	double SortDelaySeconds;
 
+	/** Weak ptr to the asset that is waiting to be renamed when scrolled into view, and the window is active */
+	TWeakPtr<struct FAssetViewItem> AwaitingRename;
+
 	/** Set when the user is in the process of naming an asset */
 	TWeakPtr<struct FAssetViewItem> RenamingAsset;
 
diff --git a/Engine/Source/Editor/DetailCustomizations/Private/AnimTransitionNodeDetails.cpp b/Engine/Source/Editor/DetailCustomizations/Private/AnimTransitionNodeDetails.cpp
--- a/Engine/Source/Editor/DetailCustomizations/Private/AnimTransitionNodeDetails.cpp
+++ b/Engine/Source/Editor/DetailCustomizations/Private/AnimTransitionNodeDetails.cpp
@@ -228,8 +228,6 @@ FReply FAnimTransitionNodeDetails::OnPromoteToSharedClick(bool RuleShare)
 			FSlateApplication::Get().GetCursorPos(),
 			FPopupTransitionEffect( FPopupTransitionEffect::TypeInPopup )
 			);
-
-		TextEntry->FocusDefaultWidget();
 		TextEntryWidget = TextEntry;
 	}
 
diff --git a/Engine/Source/Editor/DistCurveEditor/Private/SDistributionCurveEditor.cpp b/Engine/Source/Editor/DistCurveEditor/Private/SDistributionCurveEditor.cpp
--- a/Engine/Source/Editor/DistCurveEditor/Private/SDistributionCurveEditor.cpp
+++ b/Engine/Source/Editor/DistCurveEditor/Private/SDistributionCurveEditor.cpp
@@ -537,8 +537,6 @@ void SDistributionCurveEditor::OnSetTime()
 		FSlateApplication::Get().GetCursorPos(),
 		FPopupTransitionEffect(FPopupTransitionEffect::TypeInPopup)
 		);
-
-	TextEntry->FocusDefaultWidget();
 }
 
 void SDistributionCurveEditor::OnSetValue()
@@ -563,8 +561,6 @@ void SDistributionCurveEditor::OnSetValue()
 		FSlateApplication::Get().GetCursorPos(),
 		FPopupTransitionEffect(FPopupTransitionEffect::TypeInPopup)
 		);
-
-	TextEntry->FocusDefaultWidget();
 }
 
 void SDistributionCurveEditor::OnSetColor()
@@ -785,8 +781,6 @@ void SDistributionCurveEditor::OnScaleTimes(ECurveScaleScope::Type Scope)
 		FSlateApplication::Get().GetCursorPos(),
 		FPopupTransitionEffect(FPopupTransitionEffect::TypeInPopup)
 		);
-
-	TextEntry->FocusDefaultWidget();
 }
 
 void SDistributionCurveEditor::OnScaleValues(ECurveScaleScope::Type Scope)
@@ -823,8 +817,6 @@ void SDistributionCurveEditor::OnScaleValues(ECurveScaleScope::Type Scope)
 		FSlateApplication::Get().GetCursorPos(),
 		FPopupTransitionEffect(FPopupTransitionEffect::TypeInPopup)
 		);
-
-	TextEntry->FocusDefaultWidget();
 }
 
 void SDistributionCurveEditor::OnFitHorizontally()
@@ -980,8 +972,6 @@ void SDistributionCurveEditor::OnCreateTab()
 		FSlateApplication::Get().GetCursorPos(),
 		FPopupTransitionEffect(FPopupTransitionEffect::TypeInPopup)
 		);
-
-	TextEntry->FocusDefaultWidget();
 }
 
 void SDistributionCurveEditor::OnDeleteTab()
diff --git a/Engine/Source/Editor/Matinee/Private/MatineeMenus.cpp b/Engine/Source/Editor/Matinee/Private/MatineeMenus.cpp
--- a/Engine/Source/Editor/Matinee/Private/MatineeMenus.cpp
+++ b/Engine/Source/Editor/Matinee/Private/MatineeMenus.cpp
@@ -127,9 +127,7 @@ void FMatinee::GenericTextEntryModeless(const FText& DialogText, const FText& De
 		TextEntryPopup,
 		FSlateApplication::Get().GetCursorPos(),
 		FPopupTransitionEffect(FPopupTransitionEffect::TypeInPopup)
-		);	
-
-	TextEntryPopup->FocusDefaultWidget();
+		);
 }
 
 void FMatinee::CloseEntryPopupWindow()
@@ -2901,8 +2899,6 @@ void FMatinee::OnContextSaveAsCameraAnimation()
 		FSlateApplication::Get().GetCursorPos(),
 		FPopupTransitionEffect( FPopupTransitionEffect::TypeInPopup )
 		);
-
-	TextEntry->FocusDefaultWidget();
 }
 
 /**
diff --git a/Engine/Source/Editor/Matinee/Private/MatineeTrackHelpers.cpp b/Engine/Source/Editor/Matinee/Private/MatineeTrackHelpers.cpp
--- a/Engine/Source/Editor/Matinee/Private/MatineeTrackHelpers.cpp
+++ b/Engine/Source/Editor/Matinee/Private/MatineeTrackHelpers.cpp
@@ -355,8 +355,6 @@ bool UMatineeTrackDirectorHelper::PreCreateKeyframe( UInterpTrack *Track, float
 				FSlateApplication::Get().GetCursorPos(),
 				FPopupTransitionEffect(FPopupTransitionEffect::TypeInPopup)
 				);
-
-			TextEntryPopup->FocusDefaultWidget();
 		}
 	}
 
@@ -417,8 +415,6 @@ bool UMatineeTrackEventHelper::PreCreateKeyframe( UInterpTrack *Track, float Key
 			);
 	}
 
-	TextEntryPopup->FocusDefaultWidget();
-
 	return false;
 }
 
@@ -798,8 +794,6 @@ bool UMatineeTrackToggleHelper::PreCreateKeyframe( UInterpTrack *Track, float Ke
 			);
 	}
 
-	TextEntryPopup->FocusDefaultWidget();
-
 	return bResult;
 }
 
diff --git a/Engine/Source/Editor/Persona/Private/SAnimCurvePanel.cpp b/Engine/Source/Editor/Persona/Private/SAnimCurvePanel.cpp
--- a/Engine/Source/Editor/Persona/Private/SAnimCurvePanel.cpp
+++ b/Engine/Source/Editor/Persona/Private/SAnimCurvePanel.cpp
@@ -1105,8 +1105,6 @@ void SAnimCurvePanel::CreateNewMetadataEntryClicked()
 		SlateApp.GetCursorPos(),
 		FPopupTransitionEffect::TypeInPopup
 		);
-
-	TextEntry->FocusDefaultWidget();
 }
 
 void SAnimCurvePanel::CreateNewMetadataEntry(const FText& CommittedText, ETextCommit::Type CommitType)
@@ -1142,8 +1140,6 @@ void SAnimCurvePanel::CreateNewCurveClicked()
 		SlateApp.GetCursorPos(),
 		FPopupTransitionEffect::TypeInPopup
 		);
-
-	TextEntry->FocusDefaultWidget();
 }
 
 TSharedRef<SWidget> SAnimCurvePanel::CreateCurveContextMenu(FFloatCurve* Curve) const
diff --git a/Engine/Source/Editor/Persona/Private/SAnimMontagePanel.cpp b/Engine/Source/Editor/Persona/Private/SAnimMontagePanel.cpp
--- a/Engine/Source/Editor/Persona/Private/SAnimMontagePanel.cpp
+++ b/Engine/Source/Editor/Persona/Private/SAnimMontagePanel.cpp
@@ -366,8 +366,6 @@ void SAnimMontagePanel::OnNewSectionClicked(float DataPosX)
 		FSlateApplication::Get().GetCursorPos(),
 		FPopupTransitionEffect( FPopupTransitionEffect::TypeInPopup )
 		);
-
-	TextEntry->FocusDefaultWidget();
 }
 
 void SAnimMontagePanel::CreateNewSection(const FText& NewSectionName, ETextCommit::Type CommitInfo, float StartTime)
diff --git a/Engine/Source/Editor/Persona/Private/SAnimNotifyPanel.cpp b/Engine/Source/Editor/Persona/Private/SAnimNotifyPanel.cpp
--- a/Engine/Source/Editor/Persona/Private/SAnimNotifyPanel.cpp
+++ b/Engine/Source/Editor/Persona/Private/SAnimNotifyPanel.cpp
@@ -2729,8 +2729,6 @@ void SAnimNotifyTrack::OnSetTriggerWeightNotifyClicked(int32 NotifyIndex)
 			FSlateApplication::Get().GetCursorPos(),
 			FPopupTransitionEffect( FPopupTransitionEffect::TypeInPopup )
 			);
-
-		TextEntry->FocusDefaultWidget();		
 	}
 }
 
@@ -2753,8 +2751,6 @@ void SAnimNotifyTrack::OnSetDurationNotifyClicked(int32 NotifyIndex)
 			FSlateApplication::Get().GetCursorPos(),
 			FPopupTransitionEffect( FPopupTransitionEffect::TypeInPopup )
 			);
-
-		TextEntry->FocusDefaultWidget();		
 	}
 }
 
@@ -2793,8 +2789,6 @@ void SAnimNotifyTrack::OnNewNotifyClicked()
 		FSlateApplication::Get().GetCursorPos(),
 		FPopupTransitionEffect( FPopupTransitionEffect::TypeInPopup )
 		);
-
-	TextEntry->FocusDefaultWidget();
 }
 
 void SAnimNotifyTrack::AddNewNotify(const FText& NewNotifyName, ETextCommit::Type CommitInfo)
@@ -3179,8 +3173,6 @@ void SAnimNotifyTrack::OnSetNotifyTimeClicked(int32 NotifyIndex)
 			FSlateApplication::Get().GetCursorPos(),
 			FPopupTransitionEffect(FPopupTransitionEffect::TypeInPopup)
 			);
-
-		TextEntry->FocusDefaultWidget();
 	}
 }
 
@@ -3232,8 +3224,6 @@ void SAnimNotifyTrack::OnSetNotifyFrameClicked(int32 NotifyIndex)
 			FSlateApplication::Get().GetCursorPos(),
 			FPopupTransitionEffect(FPopupTransitionEffect::TypeInPopup)
 			);
-
-		TextEntry->FocusDefaultWidget();
 	}
 }
 
diff --git a/Engine/Source/Editor/Persona/Private/SSkeletonSlotNames.cpp b/Engine/Source/Editor/Persona/Private/SSkeletonSlotNames.cpp
--- a/Engine/Source/Editor/Persona/Private/SSkeletonSlotNames.cpp
+++ b/Engine/Source/Editor/Persona/Private/SSkeletonSlotNames.cpp
@@ -403,8 +403,6 @@ void SSkeletonSlotNames::OnAddSlot()
 		FSlateApplication::Get().GetCursorPos(),
 		FPopupTransitionEffect(FPopupTransitionEffect::TypeInPopup)
 		);
-
-	TextEntry->FocusDefaultWidget();
 }
 
 void SSkeletonSlotNames::OnAddGroup()
@@ -421,8 +419,6 @@ void SSkeletonSlotNames::OnAddGroup()
 		FSlateApplication::Get().GetCursorPos(),
 		FPopupTransitionEffect(FPopupTransitionEffect::TypeInPopup)
 		);
-
-	TextEntry->FocusDefaultWidget();
 }
 
 void SSkeletonSlotNames::AddSlotPopUpOnCommit(const FText & InNewSlotText, ETextCommit::Type CommitInfo)
@@ -916,8 +912,6 @@ void SSkeletonSlotNames::OnRenameSlot(FName CurrentName)
 		FSlateApplication::Get().GetCursorPos(),
 		FPopupTransitionEffect(FPopupTransitionEffect::TypeInPopup)
 		);
-
-	TextEntry->FocusDefaultWidget();
 }
 
 void SSkeletonSlotNames::RenameSlot(FName CurrentName, FName NewName)
diff --git a/Engine/Source/Editor/Persona/Private/SSkeletonSmartNameManager.cpp b/Engine/Source/Editor/Persona/Private/SSkeletonSmartNameManager.cpp
--- a/Engine/Source/Editor/Persona/Private/SSkeletonSmartNameManager.cpp
+++ b/Engine/Source/Editor/Persona/Private/SSkeletonSmartNameManager.cpp
@@ -314,8 +314,6 @@ void SSkeletonSmartNameManager::OnAddClicked()
 		SlateApp.GetCursorPos(),
 		FPopupTransitionEffect::TypeInPopup
 		);
-
-	TextEntry->FocusDefaultWidget();
 }
 
 void SSkeletonSmartNameManager::OnPostUndo()
diff --git a/Engine/Source/Editor/Sequencer/Private/Sequencer.cpp b/Engine/Source/Editor/Sequencer/Private/Sequencer.cpp
--- a/Engine/Source/Editor/Sequencer/Private/Sequencer.cpp
+++ b/Engine/Source/Editor/Sequencer/Private/Sequencer.cpp
@@ -344,8 +344,6 @@ void FSequencer::RenameShot(UMovieSceneSection* ShotSection)
 		FSlateApplication::Get().GetCursorPos(),
 		FPopupTransitionEffect( FPopupTransitionEffect::TypeInPopup )
 		);
-
-	TextEntry->FocusDefaultWidget();
 }
 
 void FSequencer::DeleteSection(class UMovieSceneSection* Section)
diff --git a/Engine/Source/Editor/StatsViewer/Private/StatsPages/StaticMeshLightingInfoStatsPage.cpp b/Engine/Source/Editor/StatsViewer/Private/StatsPages/StaticMeshLightingInfoStatsPage.cpp
--- a/Engine/Source/Editor/StatsViewer/Private/StatsPages/StaticMeshLightingInfoStatsPage.cpp
+++ b/Engine/Source/Editor/StatsViewer/Private/StatsPages/StaticMeshLightingInfoStatsPage.cpp
@@ -503,8 +503,6 @@ void FStaticMeshLightingInfoStatsPage::GetUserSetStaticLightmapResolution(TWeakP
 			FSlateApplication::Get().GetCursorPos(),
 			FPopupTransitionEffect( FPopupTransitionEffect::TypeInPopup )
 			);
-
-		TextEntry->FocusDefaultWidget();
 	}
 }
 
diff --git a/Engine/Source/Runtime/AppFramework/Private/Widgets/Testing/STestSuite.cpp b/Engine/Source/Runtime/AppFramework/Private/Widgets/Testing/STestSuite.cpp
--- a/Engine/Source/Runtime/AppFramework/Private/Widgets/Testing/STestSuite.cpp
+++ b/Engine/Source/Runtime/AppFramework/Private/Widgets/Testing/STestSuite.cpp
@@ -2117,8 +2117,6 @@ public:
 			FPopupTransitionEffect( FPopupTransitionEffect::TypeInPopup )
 		);
 
-		TextEntry->FocusDefaultWidget();
-
 		return FReply::Handled();
 	}
 
diff --git a/Engine/Source/Runtime/Slate/Private/Widgets/Input/STextComboPopup.cpp b/Engine/Source/Runtime/Slate/Private/Widgets/Input/STextComboPopup.cpp
--- a/Engine/Source/Runtime/Slate/Private/Widgets/Input/STextComboPopup.cpp
+++ b/Engine/Source/Runtime/Slate/Private/Widgets/Input/STextComboPopup.cpp
@@ -59,6 +59,27 @@ void STextComboPopup::Construct( const FArguments& InArgs )
 	StringCombo->RefreshOptions();
 
 	StringCombo->SetSelectedItem(SelectedItem);
+
+	if (InArgs._AutoFocus)
+	{
+		RegisterActiveTimer(0.016f, FWidgetActiveTimerDelegate::CreateSP(this, &STextComboPopup::TickAutoFocus));
+	}
+}
+
+EActiveTimerReturnType STextComboPopup::TickAutoFocus(double InCurrentTime, float InDeltaTime)
+{
+	TSharedPtr<SWindow> OwnerWindow = FSlateApplication::Get().FindWidgetWindow(AsShared());
+	if (!OwnerWindow.IsValid())
+	{
+		return EActiveTimerReturnType::Stop;
+	}
+	else if (FSlateApplication::Get().HasFocusedDescendants(OwnerWindow.ToSharedRef()))
+	{
+		FocusDefaultWidget();
+		return EActiveTimerReturnType::Stop;
+	}
+
+	return EActiveTimerReturnType::Continue;
 }
 
 TSharedRef<SWidget> STextComboPopup::MakeItemWidget( TSharedPtr<FString> StringItem ) 
diff --git a/Engine/Source/Runtime/Slate/Private/Widgets/Input/STextEntryPopup.cpp b/Engine/Source/Runtime/Slate/Private/Widgets/Input/STextEntryPopup.cpp
--- a/Engine/Source/Runtime/Slate/Private/Widgets/Input/STextEntryPopup.cpp
+++ b/Engine/Source/Runtime/Slate/Private/Widgets/Input/STextEntryPopup.cpp
@@ -53,6 +53,27 @@ void STextEntryPopup::Construct( const FArguments& InArgs )
 			ErrorReporting->AsWidget()
 		];
 	}
+
+	if (InArgs._AutoFocus)
+	{
+		RegisterActiveTimer(0.016f, FWidgetActiveTimerDelegate::CreateSP(this, &STextEntryPopup::TickAutoFocus));
+	}
+}
+
+EActiveTimerReturnType STextEntryPopup::TickAutoFocus(double InCurrentTime, float InDeltaTime)
+{
+	TSharedPtr<SWindow> OwnerWindow = FSlateApplication::Get().FindWidgetWindow(AsShared());
+	if (!OwnerWindow.IsValid())
+	{
+		return EActiveTimerReturnType::Stop;
+	}
+	else if (FSlateApplication::Get().HasFocusedDescendants(OwnerWindow.ToSharedRef()))
+	{
+		FocusDefaultWidget();
+		return EActiveTimerReturnType::Stop;
+	}
+
+	return EActiveTimerReturnType::Continue;
 }
 
 void STextEntryPopup::FocusDefaultWidget()
diff --git a/Engine/Source/Runtime/Slate/Public/Widgets/Input/STextComboPopup.h b/Engine/Source/Runtime/Slate/Public/Widgets/Input/STextComboPopup.h
--- a/Engine/Source/Runtime/Slate/Public/Widgets/Input/STextComboPopup.h
+++ b/Engine/Source/Runtime/Slate/Public/Widgets/Input/STextComboPopup.h
@@ -19,7 +19,7 @@ public:
 	/** The OK button */
 	TSharedPtr< SButton > OkButton;
 
-	SLATE_BEGIN_ARGS( STextComboPopup ) {}
+	SLATE_BEGIN_ARGS( STextComboPopup ) : _AutoFocus(true) {}
 
 		/** Label, placed before combo box */
 		SLATE_TEXT_ARGUMENT( Label )
@@ -30,6 +30,9 @@ public:
 		/** Called when the text is chosen. */
 		SLATE_EVENT( FOnTextChosen, OnTextChosen )
 
+		/** When set, this widget will automatically attempt to set focus to itself when it is created, or when its owner window is activated */
+		SLATE_ARGUMENT( bool, AutoFocus )
+
 	SLATE_END_ARGS()
 
 	void Construct( const FArguments& InArgs );
@@ -53,6 +56,9 @@ private:
 
 	/** Returns the item to select in the combo popup */
 	FText GetSelectedItem () const;
+	
+	/** Called when AutoFocussing to automatically set focus to this widget */
+	EActiveTimerReturnType TickAutoFocus(double InCurrentTime, float InDeltaTime);
 
 private:
 	/** The text item selected */
diff --git a/Engine/Source/Runtime/Slate/Public/Widgets/Input/STextEntryPopup.h b/Engine/Source/Runtime/Slate/Public/Widgets/Input/STextEntryPopup.h
--- a/Engine/Source/Runtime/Slate/Public/Widgets/Input/STextEntryPopup.h
+++ b/Engine/Source/Runtime/Slate/Public/Widgets/Input/STextEntryPopup.h
@@ -10,6 +10,7 @@ public:
 	SLATE_BEGIN_ARGS( STextEntryPopup )
 		: _SelectAllTextWhenFocused( false )
 		, _MaxWidth( 0.0f )
+		, _AutoFocus(true)
 		{}
 
 		/** Label, placed before text entry box */
@@ -39,6 +40,9 @@ public:
 		/** Provide a alternative mechanism for error reporting. */
 		SLATE_ARGUMENT( TSharedPtr<class IErrorReportingWidget>, ErrorReporting )
 
+		/** When set, this widget will automatically attempt to set focus to itself when it is created, or when its owner window is activated */
+		SLATE_ARGUMENT( bool, AutoFocus )
+
 	SLATE_END_ARGS()
 
 	void Construct( const FArguments& InArgs );
@@ -62,6 +66,9 @@ protected:
 	TSharedPtr<class IErrorReportingWidget> ErrorReporting;
 
 private:
+	/** Called when AutoFocussing to automatically set focus to this widget */
+	EActiveTimerReturnType TickAutoFocus(double InCurrentTime, float InDeltaTime);
+
 	/** Widget that we want to be focused when the popup is shown  */
 	TSharedPtr<SWidget> WidgetWithDefaultFocus;
 };
