Unified curve editor frame selected/all behaviour. All operations now operate on the currently selected keys, or the whole curve where none are selected.

Hooked up the horizontal range of the sequencer curve editor to the sequencer itself

diff --git a/Engine/Source/Editor/CurveAssetEditor/Private/CurveAssetEditor.cpp b/Engine/Source/Editor/CurveAssetEditor/Private/CurveAssetEditor.cpp
--- a/Engine/Source/Editor/CurveAssetEditor/Private/CurveAssetEditor.cpp
+++ b/Engine/Source/Editor/CurveAssetEditor/Private/CurveAssetEditor.cpp
@@ -185,8 +185,7 @@ TSharedPtr<FExtender> FCurveAssetEditor::GetToolbarExtender()
 			{
 				ToolbarBuilder.AddToolBarButton(FRichCurveEditorCommands::Get().ZoomToFitHorizontal);
 				ToolbarBuilder.AddToolBarButton(FRichCurveEditorCommands::Get().ZoomToFitVertical);
-				ToolbarBuilder.AddToolBarButton(FRichCurveEditorCommands::Get().ZoomToFitAll);
-				ToolbarBuilder.AddToolBarButton(FRichCurveEditorCommands::Get().ZoomToFitSelected);
+				ToolbarBuilder.AddToolBarButton(FRichCurveEditorCommands::Get().ZoomToFit);
 			}
 			ToolbarBuilder.EndSection();
 
diff --git a/Engine/Source/Editor/DistCurveEditor/Private/CurveEditorActions.cpp b/Engine/Source/Editor/DistCurveEditor/Private/CurveEditorActions.cpp
--- a/Engine/Source/Editor/DistCurveEditor/Private/CurveEditorActions.cpp
+++ b/Engine/Source/Editor/DistCurveEditor/Private/CurveEditorActions.cpp
@@ -20,8 +20,7 @@ void FCurveEditorCommands::RegisterCommands()
 	UI_COMMAND(ScaleSingleSubCurveValues, "Scale Sub-Curve Values", "Scales the values of all points of this curve", EUserInterfaceActionType::Button, FInputChord());
 	UI_COMMAND(FitHorizontally, "Horizontal", "Fit Horizontally", EUserInterfaceActionType::Button, FInputChord());
 	UI_COMMAND(FitVertically, "Vertical", "Fit Vertically", EUserInterfaceActionType::Button, FInputChord());
-	UI_COMMAND(FitToAll, "All", "Fit to All", EUserInterfaceActionType::Button, FInputChord());
-	UI_COMMAND(FitToSelected, "Selected", "Fit to Selected", EUserInterfaceActionType::Button, FInputChord());
+	UI_COMMAND(Fit, "Fit", "Fit to Selected", EUserInterfaceActionType::Button, FInputChord());
 	UI_COMMAND(PanMode, "Pan", "Pan Mode", EUserInterfaceActionType::ToggleButton, FInputChord());
 	UI_COMMAND(ZoomMode, "Zoom", "Zoom Mode", EUserInterfaceActionType::ToggleButton, FInputChord());
 	UI_COMMAND(CurveAuto, "Auto", "Curve Auto", EUserInterfaceActionType::ToggleButton, FInputChord());
diff --git a/Engine/Source/Editor/DistCurveEditor/Private/CurveEditorActions.h b/Engine/Source/Editor/DistCurveEditor/Private/CurveEditorActions.h
--- a/Engine/Source/Editor/DistCurveEditor/Private/CurveEditorActions.h
+++ b/Engine/Source/Editor/DistCurveEditor/Private/CurveEditorActions.h
@@ -31,8 +31,7 @@ public:
 	TSharedPtr<FUICommandInfo> ScaleSingleSubCurveValues;
 	TSharedPtr<FUICommandInfo> FitHorizontally;
 	TSharedPtr<FUICommandInfo> FitVertically;
-	TSharedPtr<FUICommandInfo> FitToAll;
-	TSharedPtr<FUICommandInfo> FitToSelected;
+	TSharedPtr<FUICommandInfo> Fit;
 	TSharedPtr<FUICommandInfo> PanMode;
 	TSharedPtr<FUICommandInfo> ZoomMode;
 	TSharedPtr<FUICommandInfo> CurveAuto;
diff --git a/Engine/Source/Editor/DistCurveEditor/Private/CurveEditorViewportClient.cpp b/Engine/Source/Editor/DistCurveEditor/Private/CurveEditorViewportClient.cpp
--- a/Engine/Source/Editor/DistCurveEditor/Private/CurveEditorViewportClient.cpp
+++ b/Engine/Source/Editor/DistCurveEditor/Private/CurveEditorViewportClient.cpp
@@ -624,15 +624,7 @@ bool FCurveEditorViewportClient::InputKey(FViewport* Viewport, int32 ControllerI
 		{
 			if(!bBoxSelecting && !bBegunMoving && !bDraggingHandle)
 			{
-				CurveEditorPtr.Pin()->OnFitToAll();
-			}
-			bHandled = true;
-		}
-		else if ((Key == EKeys::D) && bCtrlDown)
-		{
-			if(!bBoxSelecting && !bBegunMoving && !bDraggingHandle)
-			{
-				CurveEditorPtr.Pin()->OnFitToSelected();
+				CurveEditorPtr.Pin()->OnFit();
 			}
 			bHandled = true;
 		}
@@ -817,13 +809,12 @@ void FCurveEditorViewportClient::Exec(const TCHAR* Cmd)
 		{
 			CurveEditorPtr.Pin()->OnFitVertically();
 		}
-		else if(FParse::Command(&Str, TEXT("FitViewToAll")))
-		{
-			CurveEditorPtr.Pin()->OnFitToAll();
-		}
-		else if(FParse::Command(&Str, TEXT("FitViewToSelected")))
+		// Multiple commands to support backwards compat. with old selected/all code paths
+		else if(FParse::Command(&Str, TEXT("FitViewToAll")) ||
+				FParse::Command(&Str, TEXT("FitViewToSelected")) ||
+				FParse::Command(&Str, TEXT("FitView")))
 		{
-			CurveEditorPtr.Pin()->OnFitToSelected();
+			CurveEditorPtr.Pin()->OnFit();
 		}
 	}
 }
diff --git a/Engine/Source/Editor/DistCurveEditor/Private/SDistributionCurveEditor.cpp b/Engine/Source/Editor/DistCurveEditor/Private/SDistributionCurveEditor.cpp
--- a/Engine/Source/Editor/DistCurveEditor/Private/SDistributionCurveEditor.cpp
+++ b/Engine/Source/Editor/DistCurveEditor/Private/SDistributionCurveEditor.cpp
@@ -274,8 +274,7 @@ TSharedRef<SHorizontalBox> SDistributionCurveEditor::BuildToolBar()
 	{
 		ToolbarBuilder.AddToolBarButton(FCurveEditorCommands::Get().FitHorizontally);
 		ToolbarBuilder.AddToolBarButton(FCurveEditorCommands::Get().FitVertically);
-		ToolbarBuilder.AddToolBarButton(FCurveEditorCommands::Get().FitToAll);
-		ToolbarBuilder.AddToolBarButton(FCurveEditorCommands::Get().FitToSelected);
+		ToolbarBuilder.AddToolBarButton(FCurveEditorCommands::Get().Fit);
 	}
 	ToolbarBuilder.EndSection();
 
@@ -407,12 +406,8 @@ void SDistributionCurveEditor::BindCommands()
 		FExecuteAction::CreateSP(this, &SDistributionCurveEditor::OnFitVertically));
 
 	UICommandList->MapAction(
-		Commands.FitToAll,
-		FExecuteAction::CreateSP(this, &SDistributionCurveEditor::OnFitToAll));
-
-	UICommandList->MapAction(
-		Commands.FitToSelected,
-		FExecuteAction::CreateSP(this, &SDistributionCurveEditor::OnFitToSelected));
+		Commands.Fit,
+		FExecuteAction::CreateSP(this, &SDistributionCurveEditor::OnFit));
 
 	UICommandList->MapAction(
 		Commands.PanMode,
@@ -834,7 +829,7 @@ void SDistributionCurveEditor::OnFitVertically()
 	FitViewVertically();
 }
 
-void SDistributionCurveEditor::OnFitToAll()
+void SDistributionCurveEditor::OnFit()
 {
 	FitViewHorizontally();
 	FitViewVertically();
@@ -1462,23 +1457,15 @@ void SDistributionCurveEditor::FitViewHorizontally()
 	float MinIn = BIG_NUMBER;
 	float MaxIn = -BIG_NUMBER;
 
-	// Iterate over all curves to find the max/min Input values.
-	for(int32 i=0; i<SharedData->EdSetup->Tabs[SharedData->EdSetup->ActiveTab].Curves.Num(); i++)
-	{
-		FCurveEdEntry& Entry = SharedData->EdSetup->Tabs[SharedData->EdSetup->ActiveTab].Curves[i];
-		if(!CURVEEDENTRY_HIDECURVE(Entry.bHideCurve))
-		{
-			FCurveEdInterface* EdInterface = UInterpCurveEdSetup::GetCurveEdInterfacePointer(Entry);
-			if(EdInterface)
-			{
-				float EntryMinIn, EntryMaxIn;
-				EdInterface->GetInRange(EntryMinIn, EntryMaxIn);
+	IterateKeys([&](int32 KeyIndex, int32 SubCurveIndex, FCurveEdEntry& CurveEntry, FCurveEdInterface& EdInterface){
 
-				MinIn = FMath::Min<float>(EntryMinIn, MinIn);
-				MaxIn = FMath::Max<float>(EntryMaxIn, MaxIn);
-			}
-		}
-	}
+		const float KeyIn = EdInterface.GetKeyIn(KeyIndex);
+
+		// Update overall min and max
+		MinIn = FMath::Min<float>(KeyIn, MinIn);
+		MaxIn = FMath::Max<float>(KeyIn, MaxIn);
+		
+	});
 
 	float Size = MaxIn - MinIn;
 
@@ -1495,44 +1482,65 @@ void SDistributionCurveEditor::FitViewHorizontally()
 	Viewport->RefreshViewport();
 }
 
-void SDistributionCurveEditor::FitViewVertically()
+void SDistributionCurveEditor::IterateKeys(const TFunctionRef<void(int32, int32, FCurveEdEntry&, FCurveEdInterface&)>& IteratorCallback)
 {
-	float MinOut = BIG_NUMBER;
-	float MaxOut = -BIG_NUMBER;
-
-	for(int32 i=0; i<SharedData->EdSetup->Tabs[SharedData->EdSetup->ActiveTab].Curves.Num(); i++)
+	if (SharedData->SelectedKeys.Num() == 0)
 	{
-		FCurveEdEntry& Entry = SharedData->EdSetup->Tabs[SharedData->EdSetup->ActiveTab].Curves[i];
-		if(!CURVEEDENTRY_HIDECURVE(Entry.bHideCurve))
+		for (FCurveEdEntry& Entry : SharedData->EdSetup->Tabs[SharedData->EdSetup->ActiveTab].Curves)
 		{
-			FCurveEdInterface* EdInterface = UInterpCurveEdSetup::GetCurveEdInterfacePointer(Entry);
-			if(EdInterface)
+			if (CURVEEDENTRY_HIDECURVE(Entry.bHideCurve))
+			{
+				continue;
+			}
+			
+			if (FCurveEdInterface* CurveInterface = UInterpCurveEdSetup::GetCurveEdInterfacePointer(Entry))
 			{
-				float EntryMinOut = BIG_NUMBER;
-				float EntryMaxOut = -BIG_NUMBER;
-
 				// Iterate over each subcurve - only looking at points which are shown
-				int32 NumSubCurves = EdInterface->GetNumSubCurves();
-				for(int32 SubIndex = 0; SubIndex < EdInterface->GetNumSubCurves(); SubIndex++)
+				for(int32 SubIndex = 0; SubIndex < CurveInterface->GetNumSubCurves(); SubIndex++)
 				{
-					if(!CURVEEDENTRY_HIDESUBCURVE(Entry.bHideCurve, SubIndex))
+					if (CURVEEDENTRY_HIDESUBCURVE(Entry.bHideCurve, SubIndex))
 					{
-						// If we can see this curve - iterate over keys to find min and max 'out' value
-						for(int32 KeyIndex = 0; KeyIndex < EdInterface->GetNumKeys(); KeyIndex++)
-						{
-							float OutVal = EdInterface->GetKeyOut(SubIndex, KeyIndex);
-							EntryMinOut = FMath::Min<float>(EntryMinOut, OutVal);
-							EntryMaxOut = FMath::Max<float>(EntryMaxOut, OutVal);
-						}
+						continue;
 					}
-				}
 
-				// Update overall min and max
-				MinOut = FMath::Min<float>(EntryMinOut, MinOut);
-				MaxOut = FMath::Max<float>(EntryMaxOut, MaxOut);
+					// If we can see this curve - iterate over keys to find min and max 'out' value
+					for(int32 KeyIndex = 0; KeyIndex < CurveInterface->GetNumKeys(); KeyIndex++)
+					{
+						IteratorCallback(KeyIndex, SubIndex, Entry, *CurveInterface);
+					}
+				}
 			}
 		}
 	}
+	else for (FCurveEditorSelectedKey& SelKey : SharedData->SelectedKeys)
+	{
+		FCurveEdEntry& CurveEntry = SharedData->EdSetup->Tabs[ SharedData->EdSetup->ActiveTab ].Curves[ SelKey.CurveIndex ];
+		
+		if (CURVEEDENTRY_HIDESUBCURVE(CurveEntry.bHideCurve, SelKey.SubIndex))
+		{
+			continue;
+		}
+		else if (FCurveEdInterface* CurveInterface = UInterpCurveEdSetup::GetCurveEdInterfacePointer(CurveEntry))
+		{
+			IteratorCallback(SelKey.KeyIndex, SelKey.SubIndex, CurveEntry, *CurveInterface);
+		}
+	}
+}
+
+void SDistributionCurveEditor::FitViewVertically()
+{
+	float MinOut = BIG_NUMBER;
+	float MaxOut = -BIG_NUMBER;
+
+	IterateKeys([&](int32 KeyIndex, int32 SubCurveIndex, FCurveEdEntry& CurveEntry, FCurveEdInterface& EdInterface){
+
+		const float KeyOut = EdInterface.GetKeyOut(SubCurveIndex, KeyIndex);
+
+		// Update overall min and max
+		MinOut = FMath::Min<float>(KeyOut, MinOut);
+		MaxOut = FMath::Max<float>(KeyOut, MaxOut);
+
+	});
 
 	float Size = MaxOut - MinOut;
 
diff --git a/Engine/Source/Editor/DistCurveEditor/Private/SDistributionCurveEditor.h b/Engine/Source/Editor/DistCurveEditor/Private/SDistributionCurveEditor.h
--- a/Engine/Source/Editor/DistCurveEditor/Private/SDistributionCurveEditor.h
+++ b/Engine/Source/Editor/DistCurveEditor/Private/SDistributionCurveEditor.h
@@ -75,7 +75,7 @@ public:
 
 	/** Toolbar/menu command methods */
 	void OnDeleteKeys();
-	void OnFitToAll();
+	void OnFit();
 	void OnFitToSelected();
 	void OnFitHorizontally();
 	void OnFitVertically();
@@ -147,6 +147,9 @@ private:
 	/** Helper method to set selected tab */
 	TSharedPtr<FString> GetSelectedTab() const;
 
+	/** Helper function to iterate all selected curve keys if any are selected, otherwise all the keys in all the curves */
+	void IterateKeys(const TFunctionRef<void(int32, int32, FCurveEdEntry&, FCurveEdInterface&)>& IteratorCallback);
+
 private:
 	/** A list commands to execute if a user presses the corresponding keybinding in the text box */
 	TSharedRef<FUICommandList> UICommandList;
diff --git a/Engine/Source/Editor/EditorStyle/Private/SlateEditorStyle.cpp b/Engine/Source/Editor/EditorStyle/Private/SlateEditorStyle.cpp
--- a/Engine/Source/Editor/EditorStyle/Private/SlateEditorStyle.cpp
+++ b/Engine/Source/Editor/EditorStyle/Private/SlateEditorStyle.cpp
@@ -6003,8 +6003,7 @@ void FSlateEditorStyle::FStyle::SetupToolkitStyles()
 
 		Set( "CurveEditor.FitHorizontally", new IMAGE_BRUSH( "Icons/icon_CurveEditor_Horizontal_40x", Icon40x40 ) );
 		Set( "CurveEditor.FitVertically", new IMAGE_BRUSH( "Icons/icon_CurveEditor_Vertical_40x", Icon40x40 ) );
-		Set( "CurveEditor.FitToAll", new IMAGE_BRUSH( "Icons/icon_CurveEditor_All_40x", Icon40x40 ) );
-		Set( "CurveEditor.FitToSelected", new IMAGE_BRUSH( "Icons/icon_CurveEditor_Selected_40x", Icon40x40 ) );
+		Set( "CurveEditor.Fit", new IMAGE_BRUSH( "Icons/icon_CurveEditor_ZoomToFit_40x", Icon40x40 ) );
 		Set( "CurveEditor.PanMode", new IMAGE_BRUSH( "Icons/icon_CurveEditor_Pan_40x", Icon40x40 ) );
 		Set( "CurveEditor.ZoomMode", new IMAGE_BRUSH( "Icons/icon_CurveEditor_Zoom_40x", Icon40x40 ) );
 		Set( "CurveEditor.CurveAuto", new IMAGE_BRUSH( "Icons/icon_CurveEditor_Auto_40x", Icon40x40 ) );
@@ -6021,8 +6020,7 @@ void FSlateEditorStyle::FStyle::SetupToolkitStyles()
 		
 		Set( "CurveEditor.FitHorizontally.Small", new IMAGE_BRUSH( "Icons/icon_CurveEditor_Horizontal_40x", Icon20x20 ) );
 		Set( "CurveEditor.FitVertically.Small", new IMAGE_BRUSH( "Icons/icon_CurveEditor_Vertical_40x", Icon20x20 ) );
-		Set( "CurveEditor.FitToAll.Small", new IMAGE_BRUSH( "Icons/icon_CurveEditor_All_40x", Icon20x20 ) );
-		Set( "CurveEditor.FitToSelected.Small", new IMAGE_BRUSH( "Icons/icon_CurveEditor_Selected_40x", Icon20x20 ) );
+		Set( "CurveEditor.Fit.Small", new IMAGE_BRUSH( "Icons/icon_CurveEditor_ZoomToFit_40x", Icon20x20 ) );
 		Set( "CurveEditor.PanMode.Small", new IMAGE_BRUSH( "Icons/icon_CurveEditor_Pan_40x", Icon20x20 ) );
 		Set( "CurveEditor.ZoomMode.Small", new IMAGE_BRUSH( "Icons/icon_CurveEditor_Zoom_40x", Icon20x20 ) );
 		Set( "CurveEditor.CurveAuto.Small", new IMAGE_BRUSH( "Icons/icon_CurveEditor_Auto_40x", Icon20x20 ) );
@@ -6044,10 +6042,8 @@ void FSlateEditorStyle::FStyle::SetupToolkitStyles()
 		Set("RichCurveEditor.ZoomToFitHorizontal.Small", new IMAGE_BRUSH("Icons/icon_CurveEditor_Horizontal_40x", Icon20x20));
 		Set("RichCurveEditor.ZoomToFitVertical", new IMAGE_BRUSH("Icons/icon_CurveEditor_Vertical_40x", Icon40x40));
 		Set("RichCurveEditor.ZoomToFitVertical.Small", new IMAGE_BRUSH("Icons/icon_CurveEditor_Vertical_40x", Icon20x20));
-		Set("RichCurveEditor.ZoomToFitAll", new IMAGE_BRUSH("Icons/icon_CurveEditor_All_40x", Icon40x40));
-		Set("RichCurveEditor.ZoomToFitAll.Small", new IMAGE_BRUSH("Icons/icon_CurveEditor_All_40x", Icon20x20));
-		Set("RichCurveEditor.ZoomToFitSelected", new IMAGE_BRUSH("Icons/icon_CurveEditor_Selected_40x", Icon40x40));
-		Set("RichCurveEditor.ZoomToFitSelected.Small", new IMAGE_BRUSH("Icons/icon_CurveEditor_Selected_40x", Icon20x20));
+		Set("RichCurveEditor.ZoomToFit", new IMAGE_BRUSH("Icons/icon_CurveEditor_ZoomToFit_40x", Icon40x40));
+		Set("RichCurveEditor.ZoomToFit.Small", new IMAGE_BRUSH("Icons/icon_CurveEditor_ZoomToFit_40x", Icon20x20));
 
 		Set("RichCurveEditor.ToggleSnapping", new IMAGE_BRUSH("Icons/icon_CurveEditor_ToggleSnap_40x", Icon40x40));
 		Set("RichCurveEditor.ToggleSnapping.Small", new IMAGE_BRUSH("Icons/icon_CurveEditor_ToggleSnap_40x", Icon20x20));
diff --git a/Engine/Source/Editor/Sequencer/Private/SSequencer.cpp b/Engine/Source/Editor/Sequencer/Private/SSequencer.cpp
--- a/Engine/Source/Editor/Sequencer/Private/SSequencer.cpp
+++ b/Engine/Source/Editor/Sequencer/Private/SSequencer.cpp
@@ -192,6 +192,8 @@ void SSequencer::Construct( const FArguments& InArgs, TSharedRef< class FSequenc
 
 	USelection::SelectionChangedEvent.AddSP(this, &SSequencer::OnActorSelectionChanged);
 
+	GetMutableDefault<USequencerSettings>()->GetOnShowCurveEditorChanged()->AddSP(this, &SSequencer::OnCurveEditorVisibilityChanged);
+
 	// Create a node tree which contains a tree of movie scene data to display in the sequence
 	SequencerNodeTree = MakeShareable( new FSequencerNodeTree( InSequencer.Get() ) );
 
@@ -234,7 +236,8 @@ void SSequencer::Construct( const FArguments& InArgs, TSharedRef< class FSequenc
 
 	SAssignNew( CurveEditor, SSequencerCurveEditor, PinnedSequencer, TimeSliderController )
 		.Visibility( this, &SSequencer::GetCurveEditorVisibility )
-		.ViewRange( FAnimatedRange::WrapAttribute(InArgs._ViewRange) );
+		.OnViewRangeChanged( InArgs._OnViewRangeChanged )
+		.ViewRange( InArgs._ViewRange );
 
 	const int32 				Column0	= 0,	Column1	= 1;
 	const int32 Row0	= 0,
@@ -632,6 +635,7 @@ TSharedRef<SWidget> SSequencer::MakeTransportControls()
 SSequencer::~SSequencer()
 {
 	USelection::SelectionChangedEvent.RemoveAll(this);
+	GetMutableDefault<USequencerSettings>()->GetOnShowCurveEditorChanged()->RemoveAll(this);
 }
 
 void SSequencer::RegisterActiveTimerForPlayback()
@@ -1178,4 +1182,13 @@ EVisibility SSequencer::GetCurveEditorVisibility() const
 	return GetDefault<USequencerSettings>()->GetShowCurveEditor() ? EVisibility::Visible : EVisibility::Collapsed;
 }
 
+void SSequencer::OnCurveEditorVisibilityChanged()
+{
+	if (CurveEditor.IsValid())
+	{
+		// Only zoom horizontally if the editor is visible
+		CurveEditor->SetZoomToFit(true, GetDefault<USequencerSettings>()->GetShowCurveEditor());
+	}
+}
+
 #undef LOCTEXT_NAMESPACE
diff --git a/Engine/Source/Editor/Sequencer/Private/SSequencer.h b/Engine/Source/Editor/Sequencer/Private/SSequencer.h
--- a/Engine/Source/Editor/Sequencer/Private/SSequencer.h
+++ b/Engine/Source/Editor/Sequencer/Private/SSequencer.h
@@ -218,6 +218,9 @@ private:
 	/** Called when a column fill percentage is changed by a splitter slot. */
 	void OnColumnFillCoefficientChanged(float FillCoefficient, int32 ColumnIndex);
 
+	/** Called when the curve editor is shown or hidden */
+	void OnCurveEditorVisibilityChanged();
+
 private:
 
 	/** Section area widget */
diff --git a/Engine/Source/Editor/Sequencer/Private/SSequencerCurveEditor.cpp b/Engine/Source/Editor/Sequencer/Private/SSequencerCurveEditor.cpp
--- a/Engine/Source/Editor/Sequencer/Private/SSequencerCurveEditor.cpp
+++ b/Engine/Source/Editor/Sequencer/Private/SSequencerCurveEditor.cpp
@@ -9,17 +9,22 @@
 
 void SSequencerCurveEditor::Construct( const FArguments& InArgs, TSharedRef<FSequencer> InSequencer, TSharedRef<ITimeSliderController> InTimeSliderController )
 {
-	ViewRange = InArgs._ViewRange;
 	Sequencer = InSequencer;
 	TimeSliderController = InTimeSliderController;
 	NodeTreeSelectionChangedHandle = Sequencer.Pin()->GetSelection()->GetOnOutlinerNodeSelectionChanged()->AddSP(this, &SSequencerCurveEditor::NodeTreeSelectionChanged);
 	GetMutableDefault<USequencerSettings>()->GetOnCurveVisibilityChanged()->AddSP( this, &SSequencerCurveEditor::SequencerCurveVisibilityChanged );
 
+	// @todo sequencer: switch to lambda capture expressions when support is introduced?
+	auto ViewRange = InArgs._ViewRange;
+	auto OnViewRangeChanged = InArgs._OnViewRangeChanged;
+
 	SCurveEditor::Construct(
 		SCurveEditor::FArguments()
-		.ViewMinInput_Lambda( [this]{ return ViewRange.Get().GetLowerBoundValue(); } )
-		.ViewMaxInput_Lambda( [this]{ return ViewRange.Get().GetUpperBoundValue(); } )
+		.ViewMinInput_Lambda( [=]{ return ViewRange.Get().GetLowerBoundValue(); } )
+		.ViewMaxInput_Lambda( [=]{ return ViewRange.Get().GetUpperBoundValue(); } )
+		.OnSetInputViewRange_Lambda( [=](float InLowerBound, float InUpperBound){ OnViewRangeChanged.ExecuteIfBound(TRange<float>(InLowerBound, InUpperBound), EViewRangeInterpolation::Immediate); } )
 		.HideUI( false )
+		.ZoomToFitHorizontal( GetDefault<USequencerSettings>()->GetShowCurveEditor() )
 		.ShowCurveSelector( false )
 		.ShowZoomButtons( false )
 		.ShowInputGridNumbers( false )
diff --git a/Engine/Source/Editor/Sequencer/Private/SSequencerCurveEditor.h b/Engine/Source/Editor/Sequencer/Private/SSequencerCurveEditor.h
--- a/Engine/Source/Editor/Sequencer/Private/SSequencerCurveEditor.h
+++ b/Engine/Source/Editor/Sequencer/Private/SSequencerCurveEditor.h
@@ -3,6 +3,7 @@
 #pragma once
 
 #include "SCurveEditor.h"
+#include "ITimeSlider.h"
 
 /** A widget for displaying and managing an SCurveEditor in sequencer. */
 class SSequencerCurveEditor : public SCurveEditor
@@ -12,7 +13,10 @@ public:
 	{}
 
 		/** The range of time being viewed */
-		SLATE_ATTRIBUTE( TRange<float>, ViewRange )
+		SLATE_ATTRIBUTE( FAnimatedRange, ViewRange )
+
+		/** Event that is dispatched when the view range needs to be changed */
+		SLATE_EVENT( FOnViewRangeChanged, OnViewRangeChanged )
 
 	SLATE_END_ARGS()
 
@@ -58,8 +62,6 @@ private:
 	TWeakPtr<FSequencer> Sequencer;
 	/** The class responsible for time sliding on the curve editor */
 	TSharedPtr<class ITimeSliderController> TimeSliderController;
-	/** The visible time range displayed by the curve editor. */
-	TAttribute<TRange<float>> ViewRange;
 	/** The sequencer node tree which contains the key area nodes which supply the curves to edit. */
 	TSharedPtr<FSequencerNodeTree> SequencerNodeTree;
 	/** The sequencer curve owner implementation which is visualized by the SCurveEditor. */
diff --git a/Engine/Source/Editor/Sequencer/Private/SSequencerCurveEditorToolBar.cpp b/Engine/Source/Editor/Sequencer/Private/SSequencerCurveEditorToolBar.cpp
--- a/Engine/Source/Editor/Sequencer/Private/SSequencerCurveEditorToolBar.cpp
+++ b/Engine/Source/Editor/Sequencer/Private/SSequencerCurveEditorToolBar.cpp
@@ -10,8 +10,7 @@ void SSequencerCurveEditorToolBar::Construct( const FArguments& InArgs, TSharedP
 	{
 		ToolBarBuilder.AddToolBarButton( FRichCurveEditorCommands::Get().ZoomToFitHorizontal );
 		ToolBarBuilder.AddToolBarButton( FRichCurveEditorCommands::Get().ZoomToFitVertical );
-		ToolBarBuilder.AddToolBarButton( FRichCurveEditorCommands::Get().ZoomToFitAll );
-		ToolBarBuilder.AddToolBarButton( FRichCurveEditorCommands::Get().ZoomToFitSelected );
+		ToolBarBuilder.AddToolBarButton( FRichCurveEditorCommands::Get().ZoomToFit );
 	}
 	ToolBarBuilder.EndSection();
 
diff --git a/Engine/Source/Editor/Sequencer/Private/SequencerSettings.cpp b/Engine/Source/Editor/Sequencer/Private/SequencerSettings.cpp
--- a/Engine/Source/Editor/Sequencer/Private/SequencerSettings.cpp
+++ b/Engine/Source/Editor/Sequencer/Private/SequencerSettings.cpp
@@ -170,6 +170,7 @@ void USequencerSettings::SetShowCurveEditor(bool InbShowCurveEditor)
 	if (bShowCurveEditor != InbShowCurveEditor)
 	{
 		bShowCurveEditor = InbShowCurveEditor;
+		OnShowCurveEditorChanged.Broadcast();
 		SaveConfig();
 	}
 }
@@ -214,3 +215,8 @@ USequencerSettings::FOnCurveVisibilityChanged* USequencerSettings::GetOnCurveVis
 {
 	return &OnCurveVisibilityChanged;
 }
+
+USequencerSettings::FOnShowCurveEditorChanged* USequencerSettings::GetOnShowCurveEditorChanged()
+{
+	return &OnShowCurveEditorChanged;
+}
diff --git a/Engine/Source/Editor/Sequencer/Private/SequencerSettings.h b/Engine/Source/Editor/Sequencer/Private/SequencerSettings.h
--- a/Engine/Source/Editor/Sequencer/Private/SequencerSettings.h
+++ b/Engine/Source/Editor/Sequencer/Private/SequencerSettings.h
@@ -27,6 +27,7 @@ public:
 	GENERATED_UCLASS_BODY()
 
 	DECLARE_MULTICAST_DELEGATE( FOnCurveVisibilityChanged );
+	DECLARE_MULTICAST_DELEGATE( FOnShowCurveEditorChanged );
 
 	/** Gets whether or not snapping is enabled. */
 	bool GetIsSnapEnabled() const;
@@ -99,6 +100,9 @@ public:
 	/** Gets the multicast delegate which is run whenever the CurveVisibility value changes. */
 	FOnCurveVisibilityChanged* GetOnCurveVisibilityChanged();
 
+	/** Gets the multicast delegate which is run whenever the curve editor is shown/hidden. */
+	FOnShowCurveEditorChanged* GetOnShowCurveEditorChanged();
+
 protected:
 	UPROPERTY( config )
 	bool bIsSnapEnabled;
@@ -140,4 +144,5 @@ protected:
 	TEnumAsByte<ESequencerCurveVisibility::Type> CurveVisibility;
 
 	FOnCurveVisibilityChanged OnCurveVisibilityChanged;
+	FOnShowCurveEditorChanged OnShowCurveEditorChanged;
 };
\ No newline at end of file
diff --git a/Engine/Source/Editor/UnrealEd/Private/RichCurveEditorCommands.cpp b/Engine/Source/Editor/UnrealEd/Private/RichCurveEditorCommands.cpp
--- a/Engine/Source/Editor/UnrealEd/Private/RichCurveEditorCommands.cpp
+++ b/Engine/Source/Editor/UnrealEd/Private/RichCurveEditorCommands.cpp
@@ -9,8 +9,7 @@ void FRichCurveEditorCommands::RegisterCommands()
 {
 	UI_COMMAND(ZoomToFitHorizontal, "Fit Horizontal", "Zoom to Fit - Horizontal", EUserInterfaceActionType::Button, FInputChord());
 	UI_COMMAND(ZoomToFitVertical, "Fit Vertical", "Zoom to Fit - Vertical", EUserInterfaceActionType::Button, FInputChord());
-	UI_COMMAND(ZoomToFitAll, "Fit All", "Zoom to Fit", EUserInterfaceActionType::Button, FInputChord());
-	UI_COMMAND(ZoomToFitSelected, "Fit Selected", "Zoom to Fit - Selected", EUserInterfaceActionType::Button, FInputChord(EKeys::F));
+	UI_COMMAND(ZoomToFit, "Fit", "Zoom to Fit", EUserInterfaceActionType::Button, FInputChord(EKeys::F));
 
 	UI_COMMAND(ToggleSnapping, "Snapping", "Toggle Snapping", EUserInterfaceActionType::ToggleButton, FInputChord());
 
diff --git a/Engine/Source/Editor/UnrealEd/Private/SCurveEditor.cpp b/Engine/Source/Editor/UnrealEd/Private/SCurveEditor.cpp
--- a/Engine/Source/Editor/UnrealEd/Private/SCurveEditor.cpp
+++ b/Engine/Source/Editor/UnrealEd/Private/SCurveEditor.cpp
@@ -90,16 +90,13 @@ void SCurveEditor::Construct(const FArguments& InArgs)
 		FExecuteAction::CreateSP(this, &SCurveEditor::RedoAction));
 
 	Commands->MapAction(FRichCurveEditorCommands::Get().ZoomToFitHorizontal,
-		FExecuteAction::CreateSP(this, &SCurveEditor::ZoomToFitHorizontal, false));
+		FExecuteAction::CreateSP(this, &SCurveEditor::ZoomToFitHorizontal));
 
 	Commands->MapAction(FRichCurveEditorCommands::Get().ZoomToFitVertical,
-		FExecuteAction::CreateSP(this, &SCurveEditor::ZoomToFitVertical, false));
+		FExecuteAction::CreateSP(this, &SCurveEditor::ZoomToFitVertical));
 
-	Commands->MapAction(FRichCurveEditorCommands::Get().ZoomToFitAll,
-		FExecuteAction::CreateSP(this, &SCurveEditor::ZoomToFitAll, false));
-
-	Commands->MapAction(FRichCurveEditorCommands::Get().ZoomToFitSelected,
-		FExecuteAction::CreateSP(this, &SCurveEditor::ZoomToFitAll, true));
+	Commands->MapAction(FRichCurveEditorCommands::Get().ZoomToFit,
+		FExecuteAction::CreateSP(this, &SCurveEditor::ZoomToFit));
 
 	Commands->MapAction(FRichCurveEditorCommands::Get().ToggleSnapping,
 		FExecuteAction::CreateSP(this, &SCurveEditor::ToggleSnapping),
@@ -1037,12 +1034,12 @@ void SCurveEditor::SetCurveOwner(FCurveOwnerInterface* InCurveOwner, bool bCanEd
 
 	if( bZoomToFitVertical )
 	{
-		ZoomToFitVertical(false);
+		ZoomToFitVertical();
 	}
 
 	if ( bZoomToFitHorizontal )
 	{
-		ZoomToFitHorizontal(false);
+		ZoomToFitHorizontal();
 	}
 
 	CurveSelectionWidget.Pin()->SetContent(CreateCurveSelectionWidget());
@@ -1879,7 +1876,7 @@ TArray<FRichCurve*> SCurveEditor::GetCurvesToFit() const
 }
 
 
-void SCurveEditor::ZoomToFitHorizontal(bool OnlySelected)
+void SCurveEditor::ZoomToFitHorizontal()
 {
 	TArray<FRichCurve*> CurvesToFit = GetCurvesToFit();
 
@@ -1889,7 +1886,7 @@ void SCurveEditor::ZoomToFitHorizontal(bool OnlySelected)
 		float InMax = -FLT_MAX;
 		int32 TotalKeys = 0;
 
-		if (OnlySelected)
+		if (SelectedKeys.Num())
 		{
 			for (auto SelectedKey : SelectedKeys)
 			{
@@ -1901,9 +1898,8 @@ void SCurveEditor::ZoomToFitHorizontal(bool OnlySelected)
 		}
 		else
 		{
-			for(auto It = CurvesToFit.CreateConstIterator();It;++It)
+			for (FRichCurve* Curve : CurvesToFit)
 			{
-				FRichCurve* Curve = *It;
 				float MinTime, MaxTime;
 				Curve->GetTimeRange(MinTime, MaxTime);
 				InMin = FMath::Min(MinTime, InMin);
@@ -1939,7 +1935,7 @@ void SCurveEditor::ZoomToFitHorizontal(bool OnlySelected)
 
 FReply SCurveEditor::ZoomToFitHorizontalClicked()
 {
-	ZoomToFitHorizontal(false);
+	ZoomToFitHorizontal();
 	return FReply::Handled();
 }
 
@@ -1952,7 +1948,7 @@ void SCurveEditor::SetDefaultOutput(const float MinZoomRange)
 	SetOutputMinMax(NewMinOutput, NewMaxOutput);
 }
 
-void SCurveEditor::ZoomToFitVertical(bool OnlySelected)
+void SCurveEditor::ZoomToFitVertical()
 {
 	TArray<FRichCurve*> CurvesToFit = GetCurvesToFit();
 
@@ -1962,7 +1958,7 @@ void SCurveEditor::ZoomToFitVertical(bool OnlySelected)
 		float InMax = -FLT_MAX;
 		int32 TotalKeys = 0;
 
-		if (OnlySelected)
+		if (SelectedKeys.Num() != 0)
 		{
 			for (auto SelectedKey : SelectedKeys)
 			{
@@ -1974,9 +1970,8 @@ void SCurveEditor::ZoomToFitVertical(bool OnlySelected)
 		}
 		else
 		{
-			for(auto It = CurvesToFit.CreateConstIterator();It;++It)
+			for (FRichCurve* Curve : CurvesToFit)
 			{
-				FRichCurve* Curve = *It;
 				float MinVal, MaxVal;
 				Curve->GetValueRange(MinVal, MaxVal);
 				InMin = FMath::Min(MinVal, InMin);
@@ -2007,19 +2002,14 @@ void SCurveEditor::ZoomToFitVertical(bool OnlySelected)
 
 FReply SCurveEditor::ZoomToFitVerticalClicked()
 {
-	ZoomToFitVertical(false);
+	ZoomToFitVertical();
 	return FReply::Handled();
 }
 
-void SCurveEditor::ZoomToFitAll(bool OnlySelected)
+void SCurveEditor::ZoomToFit()
 {
-	if (OnlySelected && SelectedKeys.Num() == 0)
-	{
-		// Don't zoom to selected if there is no selection.
-		return;
-	}
-	ZoomToFitHorizontal(OnlySelected);
-	ZoomToFitVertical(OnlySelected);
+	ZoomToFitHorizontal();
+	ZoomToFitVertical();
 }
 
 void SCurveEditor::ToggleSnapping()
diff --git a/Engine/Source/Editor/UnrealEd/Public/RichCurveEditorCommands.h b/Engine/Source/Editor/UnrealEd/Public/RichCurveEditorCommands.h
--- a/Engine/Source/Editor/UnrealEd/Public/RichCurveEditorCommands.h
+++ b/Engine/Source/Editor/UnrealEd/Public/RichCurveEditorCommands.h
@@ -21,8 +21,7 @@ public:
 
 	TSharedPtr<FUICommandInfo> ZoomToFitHorizontal;
 	TSharedPtr<FUICommandInfo> ZoomToFitVertical;
-	TSharedPtr<FUICommandInfo> ZoomToFitAll;
-	TSharedPtr<FUICommandInfo> ZoomToFitSelected;
+	TSharedPtr<FUICommandInfo> ZoomToFit;
 	TSharedPtr<FUICommandInfo> ToggleSnapping;
 	TSharedPtr<FUICommandInfo> InterpolationConstant;
 	TSharedPtr<FUICommandInfo> InterpolationLinear;
diff --git a/Engine/Source/Editor/UnrealEd/Public/SCurveEditor.h b/Engine/Source/Editor/UnrealEd/Public/SCurveEditor.h
--- a/Engine/Source/Editor/UnrealEd/Public/SCurveEditor.h
+++ b/Engine/Source/Editor/UnrealEd/Public/SCurveEditor.h
@@ -322,9 +322,9 @@ private:
 	/* Get the curves to will be used during a fit operation */
 	TArray<FRichCurve*> GetCurvesToFit()const;
 
-	void ZoomToFitHorizontal(bool OnlySelected);
-	void ZoomToFitVertical(bool OnlySelected);
-	void ZoomToFitAll(bool OnlySelected);
+	void ZoomToFitHorizontal();
+	void ZoomToFitVertical();
+	void ZoomToFit();
 
 	FReply ZoomToFitHorizontalClicked();
 	FReply ZoomToFitVerticalClicked();
