Fixed issue with selection box rendering in the Curve Editor.

diff --git a/Engine/Source/Editor/DistCurveEditor/Private/CurveEditorViewportClient.cpp b/Engine/Source/Editor/DistCurveEditor/Private/CurveEditorViewportClient.cpp
--- a/Engine/Source/Editor/DistCurveEditor/Private/CurveEditorViewportClient.cpp
+++ b/Engine/Source/Editor/DistCurveEditor/Private/CurveEditorViewportClient.cpp
@@ -224,7 +224,7 @@ void FCurveEditorViewportClient::Draw(FViewport* Viewport, FCanvas* Canvas)
 		int32 MinY = FMath::Min(BoxStartY, BoxEndY);
 		int32 MaxX = FMath::Max(BoxStartX, BoxEndX);
 		int32 MaxY = FMath::Max(BoxStartY, BoxEndY);
-		FCanvasBoxItem BoxItem( FVector2D(MinX, MinY), FVector2D(BoxEndX-BoxStartX, BoxEndY-BoxStartY) );
+		FCanvasBoxItem BoxItem( FVector2D(MinX, MinY), FVector2D(MaxX-MinX, MaxY-MinY) );
 		BoxItem.SetColor( FLinearColor::Red );
 		Canvas->DrawItem( BoxItem );
 	}
