Fixed issue in calculating the vertical scrollbar ratio for the Curve Editor and Cascade viewports.

diff --git a/Engine/Source/Editor/Cascade/Private/CascadeEmitterCanvasClient.cpp b/Engine/Source/Editor/Cascade/Private/CascadeEmitterCanvasClient.cpp
--- a/Engine/Source/Editor/Cascade/Private/CascadeEmitterCanvasClient.cpp
+++ b/Engine/Source/Editor/Cascade/Private/CascadeEmitterCanvasClient.cpp
@@ -804,6 +804,11 @@ void FCascadeEmitterCanvasClient::CapturedMouseMove(FViewport* Viewport, int32 X
 
 float FCascadeEmitterCanvasClient::GetViewportVerticalScrollBarRatio() const
 {
+	if (CanvasDimensions.Y == 0.0f)
+	{
+		return 1.0f;
+	}
+
 	float WidgetHeight = 1.0f;
 	if (CascadeViewportPtr.Pin()->GetVerticalScrollBar().IsValid())
 	{
diff --git a/Engine/Source/Editor/DistCurveEditor/Private/CurveEditorViewportClient.cpp b/Engine/Source/Editor/DistCurveEditor/Private/CurveEditorViewportClient.cpp
--- a/Engine/Source/Editor/DistCurveEditor/Private/CurveEditorViewportClient.cpp
+++ b/Engine/Source/Editor/DistCurveEditor/Private/CurveEditorViewportClient.cpp
@@ -821,6 +821,11 @@ void FCurveEditorViewportClient::Exec(const TCHAR* Cmd)
 
 float FCurveEditorViewportClient::GetViewportVerticalScrollBarRatio() const
 {
+	if (SharedData->LabelContentBoxHeight == 0.0f)
+	{
+		return 1.0f;
+	}
+
 	float WidgetHeight = 1.0f;
 	if (CurveEditorViewportPtr.Pin()->GetVerticalScrollBar().IsValid())
 	{
