DistCurveEditor: Rename FCurveEditorCommands to FDistCurveEditorCommands to fix multiple symbols error

diff --git a/Engine/Source/Editor/DistCurveEditor/Private/CurveEditorActions.cpp b/Engine/Source/Editor/DistCurveEditor/Private/CurveEditorActions.cpp
--- a/Engine/Source/Editor/DistCurveEditor/Private/CurveEditorActions.cpp
+++ b/Engine/Source/Editor/DistCurveEditor/Private/CurveEditorActions.cpp
@@ -4,7 +4,7 @@
 
 #define LOCTEXT_NAMESPACE "CurveEditorCommands"
 
-void FCurveEditorCommands::RegisterCommands()
+void FDistCurveEditorCommands::RegisterCommands()
 {
 	UI_COMMAND(RemoveCurve, "Remove Curve", "Remove Curve", EUserInterfaceActionType::Button, FInputChord());
 	UI_COMMAND(RemoveAllCurves, "Remove All Curves", "Remove All Curves", EUserInterfaceActionType::Button, FInputChord());
diff --git a/Engine/Source/Editor/DistCurveEditor/Private/CurveEditorActions.h b/Engine/Source/Editor/DistCurveEditor/Private/CurveEditorActions.h
--- a/Engine/Source/Editor/DistCurveEditor/Private/CurveEditorActions.h
+++ b/Engine/Source/Editor/DistCurveEditor/Private/CurveEditorActions.h
@@ -7,15 +7,15 @@
 #include "EditorStyleSet.h"
 
 /*-----------------------------------------------------------------------------
-   FCurveEditorCommands
+   FDistCurveEditorCommands
 -----------------------------------------------------------------------------*/
 
-class FCurveEditorCommands : public TCommands<FCurveEditorCommands>
+class FDistCurveEditorCommands : public TCommands<FDistCurveEditorCommands>
 {
 public:
 	/** Constructor */
-	FCurveEditorCommands() 
-		: TCommands<FCurveEditorCommands>("CurveEditor", NSLOCTEXT("Contexts", "CurveEditor", "CurveEditor"), NAME_None, FEditorStyle::GetStyleSetName())
+	FDistCurveEditorCommands() 
+		: TCommands<FDistCurveEditorCommands>("CurveEditor", NSLOCTEXT("Contexts", "CurveEditor", "CurveEditor"), NAME_None, FEditorStyle::GetStyleSetName())
 	{
 	}
 
diff --git a/Engine/Source/Editor/DistCurveEditor/Private/SDistributionCurveEditor.cpp b/Engine/Source/Editor/DistCurveEditor/Private/SDistributionCurveEditor.cpp
--- a/Engine/Source/Editor/DistCurveEditor/Private/SDistributionCurveEditor.cpp
+++ b/Engine/Source/Editor/DistCurveEditor/Private/SDistributionCurveEditor.cpp
@@ -47,7 +47,7 @@ void SDistributionCurveEditor::Construct(const FArguments& InArgs)
 	SharedData->NotifyObject = InArgs._NotifyObject;
 
 	// Register our commands. This will only register them if not previously registered
-	FCurveEditorCommands::Register();
+	FDistCurveEditorCommands::Register();
 
 	for (int32 TabIdx = 0; TabIdx < SharedData->EdSetup->Tabs.Num(); TabIdx++)
 	{
@@ -284,42 +284,42 @@ TSharedRef<SHorizontalBox> SDistributionCurveEditor::BuildToolBar()
 	FToolBarBuilder ToolbarBuilder( UICommandList, FMultiBoxCustomization::None );
 	ToolbarBuilder.BeginSection("CurveEditorFit");
 	{
-		ToolbarBuilder.AddToolBarButton(FCurveEditorCommands::Get().FitHorizontally);
-		ToolbarBuilder.AddToolBarButton(FCurveEditorCommands::Get().FitVertically);
-		ToolbarBuilder.AddToolBarButton(FCurveEditorCommands::Get().Fit);
+		ToolbarBuilder.AddToolBarButton(FDistCurveEditorCommands::Get().FitHorizontally);
+		ToolbarBuilder.AddToolBarButton(FDistCurveEditorCommands::Get().FitVertically);
+		ToolbarBuilder.AddToolBarButton(FDistCurveEditorCommands::Get().Fit);
 	}
 	ToolbarBuilder.EndSection();
 
 	ToolbarBuilder.BeginSection("CurveEditorMode");
 	{
-		ToolbarBuilder.AddToolBarButton(FCurveEditorCommands::Get().PanMode);
-		ToolbarBuilder.AddToolBarButton(FCurveEditorCommands::Get().ZoomMode);
+		ToolbarBuilder.AddToolBarButton(FDistCurveEditorCommands::Get().PanMode);
+		ToolbarBuilder.AddToolBarButton(FDistCurveEditorCommands::Get().ZoomMode);
 	}
 	ToolbarBuilder.EndSection();
 
 	ToolbarBuilder.BeginSection("CurveEditorTangentTypes");
 	{
-		ToolbarBuilder.AddToolBarButton(FCurveEditorCommands::Get().CurveAuto);
-		ToolbarBuilder.AddToolBarButton(FCurveEditorCommands::Get().CurveAutoClamped);
-		ToolbarBuilder.AddToolBarButton(FCurveEditorCommands::Get().CurveUser);
-		ToolbarBuilder.AddToolBarButton(FCurveEditorCommands::Get().CurveBreak);
-		ToolbarBuilder.AddToolBarButton(FCurveEditorCommands::Get().Linear);
-		ToolbarBuilder.AddToolBarButton(FCurveEditorCommands::Get().Constant);
+		ToolbarBuilder.AddToolBarButton(FDistCurveEditorCommands::Get().CurveAuto);
+		ToolbarBuilder.AddToolBarButton(FDistCurveEditorCommands::Get().CurveAutoClamped);
+		ToolbarBuilder.AddToolBarButton(FDistCurveEditorCommands::Get().CurveUser);
+		ToolbarBuilder.AddToolBarButton(FDistCurveEditorCommands::Get().CurveBreak);
+		ToolbarBuilder.AddToolBarButton(FDistCurveEditorCommands::Get().Linear);
+		ToolbarBuilder.AddToolBarButton(FDistCurveEditorCommands::Get().Constant);
 	}
 	ToolbarBuilder.EndSection();
 
 	ToolbarBuilder.BeginSection("CurveEditorTangentOptions");
 	{
-		ToolbarBuilder.AddToolBarButton(FCurveEditorCommands::Get().FlattenTangents);
-		ToolbarBuilder.AddToolBarButton(FCurveEditorCommands::Get().StraightenTangents);
-		ToolbarBuilder.AddToolBarButton(FCurveEditorCommands::Get().ShowAllTangents);
+		ToolbarBuilder.AddToolBarButton(FDistCurveEditorCommands::Get().FlattenTangents);
+		ToolbarBuilder.AddToolBarButton(FDistCurveEditorCommands::Get().StraightenTangents);
+		ToolbarBuilder.AddToolBarButton(FDistCurveEditorCommands::Get().ShowAllTangents);
 	}
 	ToolbarBuilder.EndSection();
 
 	ToolbarBuilder.BeginSection("CurveEditorTabs");
 	{
-		ToolbarBuilder.AddToolBarButton(FCurveEditorCommands::Get().CreateTab);
-		ToolbarBuilder.AddToolBarButton(FCurveEditorCommands::Get().DeleteTab);
+		ToolbarBuilder.AddToolBarButton(FDistCurveEditorCommands::Get().CreateTab);
+		ToolbarBuilder.AddToolBarButton(FDistCurveEditorCommands::Get().DeleteTab);
 		ToolbarBuilder.AddWidget(
 			SNew(SBox)
 			.WidthOverride(175)
@@ -363,7 +363,7 @@ TSharedRef<SHorizontalBox> SDistributionCurveEditor::BuildToolBar()
 
 void SDistributionCurveEditor::BindCommands()
 {
-	const FCurveEditorCommands& Commands = FCurveEditorCommands::Get();
+	const FDistCurveEditorCommands& Commands = FDistCurveEditorCommands::Get();
 
 	UICommandList->MapAction(
 		Commands.RemoveCurve,
@@ -1075,8 +1075,8 @@ TSharedRef<SWidget> SDistributionCurveEditor::BuildMenuWidgetLabel()
 	const bool bShouldCloseWindowAfterMenuSelection = true;	// Set the menu to automatically close when the user commits to a choice
 	FMenuBuilder MenuBuilder(bShouldCloseWindowAfterMenuSelection, UICommandList);
 	{
-		MenuBuilder.AddMenuEntry(FCurveEditorCommands::Get().RemoveCurve);
-		MenuBuilder.AddMenuEntry(FCurveEditorCommands::Get().RemoveAllCurves);
+		MenuBuilder.AddMenuEntry(FDistCurveEditorCommands::Get().RemoveCurve);
+		MenuBuilder.AddMenuEntry(FDistCurveEditorCommands::Get().RemoveAllCurves);
 	}
 
 	return MenuBuilder.MakeWidget();
@@ -1091,8 +1091,8 @@ TSharedRef<SWidget> SDistributionCurveEditor::BuildMenuWidgetKey()
 		{
 			if(SharedData->SelectedKeys.Num() == 1)
 			{
-				MenuBuilder.AddMenuEntry(FCurveEditorCommands::Get().SetTime);
-				MenuBuilder.AddMenuEntry(FCurveEditorCommands::Get().SetValue);
+				MenuBuilder.AddMenuEntry(FDistCurveEditorCommands::Get().SetTime);
+				MenuBuilder.AddMenuEntry(FDistCurveEditorCommands::Get().SetValue);
 
 				FCurveEditorSelectedKey& SelKey = SharedData->SelectedKeys[0];
 				FCurveEdEntry& Entry = SharedData->EdSetup->Tabs[SharedData->EdSetup->ActiveTab].Curves[SelKey.CurveIndex];
@@ -1100,7 +1100,7 @@ TSharedRef<SWidget> SDistributionCurveEditor::BuildMenuWidgetKey()
 
 				if(Entry.bColorCurve && EdInterface->GetNumSubCurves() == 3)
 				{
-					MenuBuilder.AddMenuEntry(FCurveEditorCommands::Get().SetColor);
+					MenuBuilder.AddMenuEntry(FDistCurveEditorCommands::Get().SetColor);
 				}
 			}
 		}
@@ -1108,7 +1108,7 @@ TSharedRef<SWidget> SDistributionCurveEditor::BuildMenuWidgetKey()
 
 		MenuBuilder.BeginSection("DistributionCurveWidgetKey2");
 		{
-			MenuBuilder.AddMenuEntry(FCurveEditorCommands::Get().DeleteKeys);
+			MenuBuilder.AddMenuEntry(FDistCurveEditorCommands::Get().DeleteKeys);
 		}
 		MenuBuilder.EndSection();
 	}
@@ -1123,8 +1123,8 @@ TSharedRef<SWidget> SDistributionCurveEditor::BuildMenuWidgetGeneral()
 
 	MenuBuilder.BeginSection("AllCurvesSection", LOCTEXT("AllCurvesMenuHeader", "All Curves"));
 	{
-		MenuBuilder.AddMenuEntry(FCurveEditorCommands::Get().ScaleTimes);
-		MenuBuilder.AddMenuEntry(FCurveEditorCommands::Get().ScaleValues);
+		MenuBuilder.AddMenuEntry(FDistCurveEditorCommands::Get().ScaleTimes);
+		MenuBuilder.AddMenuEntry(FDistCurveEditorCommands::Get().ScaleValues);
 	}
 	MenuBuilder.EndSection();
 
@@ -1138,21 +1138,21 @@ TSharedRef<SWidget> SDistributionCurveEditor::BuildMenuWidgetCurve()
 	{
 		MenuBuilder.BeginSection("AllCurvesSection", LOCTEXT("AllCurvesMenuHeader", "All Curves"));
 		{
-			MenuBuilder.AddMenuEntry(FCurveEditorCommands::Get().ScaleTimes);
-			MenuBuilder.AddMenuEntry(FCurveEditorCommands::Get().ScaleValues);
+			MenuBuilder.AddMenuEntry(FDistCurveEditorCommands::Get().ScaleTimes);
+			MenuBuilder.AddMenuEntry(FDistCurveEditorCommands::Get().ScaleValues);
 		}
 		MenuBuilder.EndSection();
 
 		MenuBuilder.BeginSection("CurrentCurveSection", LOCTEXT("CurrentCurveMenuHeader", "Current Curve"));
 		{
-			MenuBuilder.AddMenuEntry(FCurveEditorCommands::Get().ScaleSingleCurveTimes);
-			MenuBuilder.AddMenuEntry(FCurveEditorCommands::Get().ScaleSingleCurveValues);
+			MenuBuilder.AddMenuEntry(FDistCurveEditorCommands::Get().ScaleSingleCurveTimes);
+			MenuBuilder.AddMenuEntry(FDistCurveEditorCommands::Get().ScaleSingleCurveValues);
 		}
 		MenuBuilder.EndSection();
 
 		MenuBuilder.BeginSection("SubCurveSection", LOCTEXT("SubCurveMenuHeader", "Sub-Curve"));
 		{
-			MenuBuilder.AddMenuEntry(FCurveEditorCommands::Get().ScaleSingleSubCurveValues);
+			MenuBuilder.AddMenuEntry(FDistCurveEditorCommands::Get().ScaleSingleSubCurveValues);
 		}
 		MenuBuilder.EndSection();
 	}
