Modified DistributionCurveEditor and CurveEditorViewport (used by Cascade and Matinee) to defer redraws rather than drawing immediately upon invalidate. Multiple immediate redraws may have been causing a reported hitch when changing particle emitter modules' distribution method properties inCascade.

diff --git a/Engine/Source/Editor/Cascade/Private/Cascade.cpp b/Engine/Source/Editor/Cascade/Private/Cascade.cpp
--- a/Engine/Source/Editor/Cascade/Private/Cascade.cpp
+++ b/Engine/Source/Editor/Cascade/Private/Cascade.cpp
@@ -1668,6 +1668,14 @@ void FCascade::Tick(float DeltaTime)
 		}
 	}
 
+	if (CurveEditor.IsValid())
+	{
+		if (CurveEditor->GetNeedsRedraw())
+		{
+			CurveEditor->DrawViewport();
+		}
+	}
+
 	if (bCurrentlySoloing != bIsSoloing)
 	{
 		bIsSoloing = bCurrentlySoloing;
diff --git a/Engine/Source/Editor/Cascade/Private/SCascadeEmitterCanvas.cpp b/Engine/Source/Editor/Cascade/Private/SCascadeEmitterCanvas.cpp
--- a/Engine/Source/Editor/Cascade/Private/SCascadeEmitterCanvas.cpp
+++ b/Engine/Source/Editor/Cascade/Private/SCascadeEmitterCanvas.cpp
@@ -77,7 +77,6 @@ void SCascadeEmitterCanvas::Construct(const FArguments& InArgs)
 void SCascadeEmitterCanvas::RefreshViewport()
 {
 	Viewport->Invalidate();
-	Viewport->InvalidateDisplay();
 }
 
 bool SCascadeEmitterCanvas::IsVisible() const
diff --git a/Engine/Source/Editor/DistCurveEditor/Private/CurveEditorViewportClient.cpp b/Engine/Source/Editor/DistCurveEditor/Private/CurveEditorViewportClient.cpp
--- a/Engine/Source/Editor/DistCurveEditor/Private/CurveEditorViewportClient.cpp
+++ b/Engine/Source/Editor/DistCurveEditor/Private/CurveEditorViewportClient.cpp
@@ -49,6 +49,7 @@ FCurveEditorViewportClient::FCurveEditorViewportClient(TWeakPtr<SDistributionCur
 	MovementAxisLock = AxisLock_None;
 	bBoxSelecting = false;
 	bKeyAdded = false;
+	bNeedsRedraw = true;
 	BoxStartX = 0;
 	BoxStartY = 0;
 	BoxEndX = 0;
@@ -785,6 +786,11 @@ bool FCurveEditorViewportClient::InputAxis(FViewport* Viewport, int32 Controller
 	return false;
 }
 
+void FCurveEditorViewportClient::RedrawRequested(FViewport* InViewport)
+{
+	bNeedsRedraw = true;
+}
+
 void FCurveEditorViewportClient::Exec(const TCHAR* Cmd)
 {
 	const TCHAR* Str = Cmd;
diff --git a/Engine/Source/Editor/DistCurveEditor/Private/CurveEditorViewportClient.h b/Engine/Source/Editor/DistCurveEditor/Private/CurveEditorViewportClient.h
--- a/Engine/Source/Editor/DistCurveEditor/Private/CurveEditorViewportClient.h
+++ b/Engine/Source/Editor/DistCurveEditor/Private/CurveEditorViewportClient.h
@@ -28,6 +28,7 @@ public:
 	virtual bool InputKey(FViewport* Viewport, int32 ControllerId, FKey Key, EInputEvent Event, float AmountDepressed = 1.0f, bool bGamepad = false) override;
 	virtual void MouseMove(FViewport* Viewport, int32 X, int32 Y) override;
 	virtual bool InputAxis(FViewport* Viewport, int32 ControllerId, FKey Key, float Delta, float DeltaTime, int32 NumSamples = 1, bool bGamepad = false) override;
+	virtual void RedrawRequested(FViewport* Viewport) override;
 
 	/** Exec handler */
 	void Exec(const TCHAR* Cmd);
@@ -40,6 +41,10 @@ public:
 
 	/** FViewport interface */
 	virtual float UpdateViewportClientWindowDPIScale() const override;
+
+	bool GetNeedsRedraw() { return bNeedsRedraw; }
+	void SetNeedsRedraw(bool bInNeedsRedraw) { bNeedsRedraw = bInNeedsRedraw; }
+
 private:
 	/** Curve editor key movement axis locking */
 	enum ECurveEdMovementAxisLock
@@ -128,6 +133,7 @@ private:
 	ECurveEdMovementAxisLock MovementAxisLock;
 	bool bBoxSelecting;
 	bool bKeyAdded;
+	bool bNeedsRedraw;
 	int32 DistanceDragged;
 	int32 BoxStartX;
 	int32 BoxStartY;
diff --git a/Engine/Source/Editor/DistCurveEditor/Private/SCurveEditorViewport.cpp b/Engine/Source/Editor/DistCurveEditor/Private/SCurveEditorViewport.cpp
--- a/Engine/Source/Editor/DistCurveEditor/Private/SCurveEditorViewport.cpp
+++ b/Engine/Source/Editor/DistCurveEditor/Private/SCurveEditorViewport.cpp
@@ -63,10 +63,22 @@ void SCurveEditorViewport::Construct(const FArguments& InArgs)
 	AdjustScrollBar();
 }
 
+void SCurveEditorViewport::DrawViewport()
+{
+	if (Viewport.IsValid())
+	{
+		Viewport->Draw();
+	}
+
+	if (ViewportClient.IsValid())
+	{
+		ViewportClient->SetNeedsRedraw(false);
+	}
+}
+
 void SCurveEditorViewport::RefreshViewport()
 {
 	Viewport->Invalidate();
-	Viewport->InvalidateDisplay();
 }
 
 void SCurveEditorViewport::SetVerticalScrollBarPosition(float Position)
diff --git a/Engine/Source/Editor/DistCurveEditor/Private/SCurveEditorViewport.h b/Engine/Source/Editor/DistCurveEditor/Private/SCurveEditorViewport.h
--- a/Engine/Source/Editor/DistCurveEditor/Private/SCurveEditorViewport.h
+++ b/Engine/Source/Editor/DistCurveEditor/Private/SCurveEditorViewport.h
@@ -35,6 +35,9 @@ public:
 	/** SCompoundWidget interface */
 	void Construct(const FArguments& InArgs);
 	
+	/** Draws the viewport */
+	void DrawViewport();
+
 	/** Refreshes the viewport */
 	void RefreshViewport();
 
diff --git a/Engine/Source/Editor/DistCurveEditor/Private/SDistributionCurveEditor.cpp b/Engine/Source/Editor/DistCurveEditor/Private/SDistributionCurveEditor.cpp
--- a/Engine/Source/Editor/DistCurveEditor/Private/SDistributionCurveEditor.cpp
+++ b/Engine/Source/Editor/DistCurveEditor/Private/SDistributionCurveEditor.cpp
@@ -63,7 +63,21 @@ void SDistributionCurveEditor::Construct(const FArguments& InArgs)
 void SDistributionCurveEditor::RefreshViewport()
 {
 	Viewport->GetViewport()->Invalidate();
-	Viewport->GetViewport()->InvalidateDisplay();
+}
+
+bool SDistributionCurveEditor::GetNeedsRedraw()
+{
+	if (Viewport->GetViewportClient().IsValid())
+	{
+		return Viewport->GetViewportClient()->GetNeedsRedraw();
+	}
+
+	return false;
+}
+
+void SDistributionCurveEditor::DrawViewport()
+{
+	Viewport->DrawViewport();
 }
 
 void SDistributionCurveEditor::CurveChanged()
diff --git a/Engine/Source/Editor/DistCurveEditor/Private/SDistributionCurveEditor.h b/Engine/Source/Editor/DistCurveEditor/Private/SDistributionCurveEditor.h
--- a/Engine/Source/Editor/DistCurveEditor/Private/SDistributionCurveEditor.h
+++ b/Engine/Source/Editor/DistCurveEditor/Private/SDistributionCurveEditor.h
@@ -61,6 +61,8 @@ public:
 	void Construct(const FArguments& InArgs);
 
 	/** IDistributionCurveEditor interface */
+	virtual void DrawViewport() override;
+	virtual bool GetNeedsRedraw() override;
 	virtual void RefreshViewport() override;
 	virtual void CurveChanged() override;
 	virtual void SetCurveVisible(const UObject* InCurve, bool bShow) override;
diff --git a/Engine/Source/Editor/DistCurveEditor/Public/IDistCurveEditor.h b/Engine/Source/Editor/DistCurveEditor/Public/IDistCurveEditor.h
--- a/Engine/Source/Editor/DistCurveEditor/Public/IDistCurveEditor.h
+++ b/Engine/Source/Editor/DistCurveEditor/Public/IDistCurveEditor.h
@@ -38,7 +38,13 @@ public:
 class DISTCURVEEDITOR_API IDistributionCurveEditor : public SCompoundWidget
 {
 public:
-	/** Refreshes the viewport */
+	/** Draws the viewport */
+	virtual void DrawViewport() = 0;
+
+	/** Returns whether draw is needed */
+	virtual bool GetNeedsRedraw() = 0;
+
+	/** Refreshes the viewport by invalidating */
 	virtual void RefreshViewport() = 0;
 
 	/** Clears selected keys and updates the viewport */
diff --git a/Engine/Source/Editor/Matinee/Private/Matinee.cpp b/Engine/Source/Editor/Matinee/Private/Matinee.cpp
--- a/Engine/Source/Editor/Matinee/Private/Matinee.cpp
+++ b/Engine/Source/Editor/Matinee/Private/Matinee.cpp
@@ -439,6 +439,7 @@ void FMatinee::BuildCurveEditor()
 	CurveEd->SetEndMarker(true, IData->InterpLength);
 	CurveEd->SetPositionMarker(true, 0.f, PosMarkerColor);
 	CurveEd->SetRegionMarker(true, IData->EdSectionStart, IData->EdSectionEnd, RegionFillColor);
+	CurveEd->DrawViewport();
 };
 
 /** Should NOT open an InterpEd unless InInterp has a valid MatineeData attached! */
@@ -2641,6 +2642,7 @@ void FMatinee::NotifyPreChange( UProperty* PropertyAboutToChange )
 void FMatinee::NotifyPostChange( const FPropertyChangedEvent& PropertyChangedEvent, UProperty* PropertyThatChanged )
 {
 	CurveEd->CurveChanged();
+	CurveEd->DrawViewport();
 
 	// Dirty the track window viewports
 	InvalidateTrackWindowViewports();
