Fixed a bug where panning was no longer possible in the curve editor widget

diff --git a/Engine/Source/Editor/DistCurveEditor/Private/CurveEditorViewportClient.cpp b/Engine/Source/Editor/DistCurveEditor/Private/CurveEditorViewportClient.cpp
--- a/Engine/Source/Editor/DistCurveEditor/Private/CurveEditorViewportClient.cpp
+++ b/Engine/Source/Editor/DistCurveEditor/Private/CurveEditorViewportClient.cpp
@@ -695,7 +695,7 @@ bool FCurveEditorViewportClient::InputKey(FViewport* Viewport, int32 ControllerI
 }
 
 void FCurveEditorViewportClient::MouseMove(FViewport* Viewport, int32 X, int32 Y)
-{	
+{
 	bool bCtrlDown = Viewport->KeyState(EKeys::LeftControl) || Viewport->KeyState(EKeys::RightControl);
 	bool bShiftDown = Viewport->KeyState(EKeys::LeftShift) || Viewport->KeyState(EKeys::RightShift);
 
@@ -796,6 +796,11 @@ void FCurveEditorViewportClient::MouseMove(FViewport* Viewport, int32 X, int32 Y
 	Viewport->InvalidateDisplay();
 }
 
+void FCurveEditorViewportClient::CapturedMouseMove(FViewport* Viewport, int32 X, int32 Y)
+{
+	MouseMove(Viewport, X, Y);
+}
+
 bool FCurveEditorViewportClient::InputAxis(FViewport* Viewport, int32 ControllerId, FKey Key, float Delta, float DeltaTime, int32 NumSamples, bool bGamepad)
 {
 	if (Key == EKeys::MouseX || Key == EKeys::MouseY)
diff --git a/Engine/Source/Editor/DistCurveEditor/Private/CurveEditorViewportClient.h b/Engine/Source/Editor/DistCurveEditor/Private/CurveEditorViewportClient.h
--- a/Engine/Source/Editor/DistCurveEditor/Private/CurveEditorViewportClient.h
+++ b/Engine/Source/Editor/DistCurveEditor/Private/CurveEditorViewportClient.h
@@ -36,6 +36,7 @@ public:
 	virtual void Draw(FViewport* Viewport, FCanvas* Canvas) override;
 	virtual bool InputKey(FViewport* Viewport, int32 ControllerId, FKey Key, EInputEvent Event, float AmountDepressed = 1.0f, bool bGamepad = false) override;
 	virtual void MouseMove(FViewport* Viewport, int32 X, int32 Y) override;
+	virtual void CapturedMouseMove(FViewport* Viewport, int32 X, int32 Y) override;
 	virtual bool InputAxis(FViewport* Viewport, int32 ControllerId, FKey Key, float Delta, float DeltaTime, int32 NumSamples = 1, bool bGamepad = false) override;
 
 	/** Exec handler */
