Fixed up modules to include rendercore to work around linker error in FViewport.

diff --git a/Engine/Source/Editor/DistCurveEditor/DistCurveEditor.Build.cs b/Engine/Source/Editor/DistCurveEditor/DistCurveEditor.Build.cs
--- a/Engine/Source/Editor/DistCurveEditor/DistCurveEditor.Build.cs
+++ b/Engine/Source/Editor/DistCurveEditor/DistCurveEditor.Build.cs
@@ -14,7 +14,8 @@ public class DistCurveEditor : ModuleRules
 				"Core",
 				"CoreUObject",
 				"Engine",
-                "InputCore",
+				"RenderCore",
+				"InputCore",
 				"Slate",
 				"SlateCore",
 				"LevelEditor",
diff --git a/Engine/Source/Editor/FontEditor/FontEditor.Build.cs b/Engine/Source/Editor/FontEditor/FontEditor.Build.cs
--- a/Engine/Source/Editor/FontEditor/FontEditor.Build.cs
+++ b/Engine/Source/Editor/FontEditor/FontEditor.Build.cs
@@ -17,6 +17,7 @@ public class FontEditor : ModuleRules
 				"DesktopWidgets",
 				"Engine",
                 "InputCore",
+				"RenderCore",
 				"Slate",
 				"SlateCore",
 				"UnrealEd",
